VulkanRHI: Add global samplers support in Vulkan (see Common.ush: collection of common samplers that automatically get bound).  Avoids using the fallback which includes the View UB.

#rnx

diff --git a/Engine/Shaders/Public/Platform/Vulkan/VulkanCommon.ush b/Engine/Shaders/Public/Platform/Vulkan/VulkanCommon.ush
--- a/Engine/Shaders/Public/Platform/Vulkan/VulkanCommon.ush
+++ b/Engine/Shaders/Public/Platform/Vulkan/VulkanCommon.ush
@@ -7,7 +7,7 @@
 #pragma once
 
 // Update this GUID to invalidate shader recompilation for only Vulkan shaders
-#pragma message("UESHADERMETADATA_VERSION DE4FD0BE-F0E4-4B69-B942-F568445132B0")
+#pragma message("UESHADERMETADATA_VERSION 5CAB236E-91CA-4E25-9DA7-B53CF6BE4E03")
 
 #define COMPILER_SUPPORTS_ULONG_TYPES 1
 #define COMPILER_SUPPORTS_UINT64_IMAGE_ATOMICS (1 && COMPILER_SUPPORTS_ULONG_TYPES)
@@ -143,3 +143,28 @@ void ImageInterlockedAddUInt64(RWTexture2D<UlongType> Dest, uint2 Coord, UlongTy
 
 	#define MESH_SHADER_WRITE_PRIMITIVE(InPrimitiveIndex, Value) OutPrimitives[InPrimitiveIndex] = Value;
 #endif
+
+
+// Types of Global Samplers (see Common.ush for types)
+// Must match GetGlobalSamplerType() in SpirVShaderCompiler.inl
+// and EGlobalSamplerType in VulkanShaderResources.h
+
+#if PLATFORM_SUPPORTS_BINDLESS
+	typedef SamplerState SafeTypeGlobalSamplers;
+	SafeTypeGlobalSamplers VULKAN_SAMPLER_HEAP(GlobalSamplers)[];
+	static const SamplerState VulkanGlobalPointClampedSampler     = VULKAN_SAMPLER_HEAP(GlobalSamplers)[1];
+	static const SamplerState VulkanGlobalPointWrappedSampler     = VULKAN_SAMPLER_HEAP(GlobalSamplers)[2];
+	static const SamplerState VulkanGlobalBilinearClampedSampler  = VULKAN_SAMPLER_HEAP(GlobalSamplers)[3];
+	static const SamplerState VulkanGlobalBilinearWrappedSampler  = VULKAN_SAMPLER_HEAP(GlobalSamplers)[4];
+	static const SamplerState VulkanGlobalTrilinearClampedSampler = VULKAN_SAMPLER_HEAP(GlobalSamplers)[5];
+	static const SamplerState VulkanGlobalTrilinearWrappedSampler = VULKAN_SAMPLER_HEAP(GlobalSamplers)[6];
+#else
+	SamplerState VulkanGlobalPointWrappedSampler;
+	SamplerState VulkanGlobalPointClampedSampler;
+	SamplerState VulkanGlobalBilinearWrappedSampler;
+	SamplerState VulkanGlobalBilinearClampedSampler;
+	SamplerState VulkanGlobalTrilinearWrappedSampler;
+	SamplerState VulkanGlobalTrilinearClampedSampler;
+#endif
+
+#define GetGlobalSampler(Filter, WrapMode) VulkanGlobal##Filter##WrapMode##Sampler
diff --git a/Engine/Source/Developer/VulkanShaderFormat/Private/SpirVShaderCompiler.inl b/Engine/Source/Developer/VulkanShaderFormat/Private/SpirVShaderCompiler.inl
--- a/Engine/Source/Developer/VulkanShaderFormat/Private/SpirVShaderCompiler.inl
+++ b/Engine/Source/Developer/VulkanShaderFormat/Private/SpirVShaderCompiler.inl
@@ -188,6 +188,27 @@ static uint32 GetUBLayoutHash(const FShaderCompilerInput& ShaderInput, const FSt
 	return LayoutHash;
 }
 
+// Types of Global Samplers (see Common.ush for types)
+// Must match EGlobalSamplerType in VulkanShaderResources.h
+// and declarations in VulkanCommon.ush
+static FVulkanShaderHeader::EGlobalSamplerType GetGlobalSamplerType(const FString& ResourceName)
+{
+#define VULKAN_GLOBAL_SAMPLER_NAME(FilterWrapName) if (ResourceName.EndsWith(TEXT(#FilterWrapName))) return FVulkanShaderHeader::EGlobalSamplerType::FilterWrapName
+
+	if (ResourceName.StartsWith(TEXT("VulkanGlobal")))
+	{
+		VULKAN_GLOBAL_SAMPLER_NAME(PointClampedSampler);
+		VULKAN_GLOBAL_SAMPLER_NAME(PointWrappedSampler);
+		VULKAN_GLOBAL_SAMPLER_NAME(BilinearClampedSampler);
+		VULKAN_GLOBAL_SAMPLER_NAME(BilinearWrappedSampler);
+		VULKAN_GLOBAL_SAMPLER_NAME(TrilinearClampedSampler);
+		VULKAN_GLOBAL_SAMPLER_NAME(TrilinearWrappedSampler);
+	}
+	return FVulkanShaderHeader::EGlobalSamplerType::Invalid;
+
+#undef VULKAN_GLOBAL_SAMPLER_NAME
+}
+
 static bool HasDerivatives(const FSpirv& Spirv)
 {
 	for (FSpirvConstIterator Iter = Spirv.begin(); Iter != Spirv.end(); ++Iter)
@@ -739,8 +760,20 @@ static bool BuildShaderOutputFromSpirv(
 					break;
 
 				case VK_DESCRIPTOR_TYPE_SAMPLER:
-					HandleReflectedShaderSampler(ResourceName, ReflectionSlot, Output);
-					//HandleReflectedShaderSampler(ResourceName, BindingTypeCount, ReflectionSlot, 1, Output);
+					{
+						// Regular samplers need reflection to get bindings, global samplers get bound automagically.
+						FVulkanShaderHeader::EGlobalSamplerType GlobalSamplerType = GetGlobalSamplerType(ResourceName);
+						if (GlobalSamplerType == FVulkanShaderHeader::EGlobalSamplerType::Invalid)
+						{
+							HandleReflectedShaderSampler(ResourceName, ReflectionSlot, Output);
+						}
+						else
+						{
+							FVulkanShaderHeader::FGlobalSamplerInfo& GlobalSamplerInfo = SerializedOutput.Header.GlobalSamplerInfos.AddZeroed_GetRef();
+							GlobalSamplerInfo.BindingIndex = BindingSlot;
+							GlobalSamplerInfo.Type = GlobalSamplerType;
+						}
+					}
 					break;
 
 				case VK_DESCRIPTOR_TYPE_ACCELERATION_STRUCTURE_KHR:
diff --git a/Engine/Source/Runtime/VulkanRHI/Private/VulkanDevice.cpp b/Engine/Source/Runtime/VulkanRHI/Private/VulkanDevice.cpp
--- a/Engine/Source/Runtime/VulkanRHI/Private/VulkanDevice.cpp
+++ b/Engine/Source/Runtime/VulkanRHI/Private/VulkanDevice.cpp
@@ -219,7 +219,6 @@ FVulkanDevice::FVulkanDevice(FVulkanDynamicRHI* InRHI, VkPhysicalDevice InGpu)
 	: Device(VK_NULL_HANDLE)
 	, MemoryManager(this)
 	, DeferredDeletionQueue(this)
-	, DefaultSampler(nullptr)
 	, Gpu(InGpu)
 	, ImmediateContext(nullptr)
 	, ComputeContext(nullptr)
@@ -1391,6 +1390,8 @@ void FVulkanDevice::InitGPU()
 	BindlessDescriptorManager = new FVulkanBindlessDescriptorManager(*this);
 	BindlessDescriptorManager->Init();
 
+	InitGlobalSamplers();
+
 	PipelineStateCache = new FVulkanPipelineStateCacheManager(this);
 
 	TArray<FString> CacheFilenames = FVulkanPlatform::GetPSOCacheFilenames();
@@ -1431,12 +1432,6 @@ void FVulkanDevice::InitGPU()
 
 	PipelineStateCache->InitAndLoad(CacheFilenames);
 
-	// Setup default resource
-	{
-		FSamplerStateInitializerRHI Default(SF_Point);
-		DefaultSampler = ResourceCast(RHICreateSamplerState(Default).GetReference());
-	}
-
 	if (RHISupportsRayTracing(GMaxRHIShaderPlatform) && GetOptionalExtensions().HasRaytracingExtensions())
 	{
 		check(RayTracingCompactionRequestHandler == nullptr);
@@ -1487,9 +1482,6 @@ void FVulkanDevice::Destroy()
 	delete DescriptorPoolsManager;
 	DescriptorPoolsManager = nullptr;
 
-	// No need to delete as it's stored in SamplerMap
-	DefaultSampler = nullptr;
-
 	for (int32 Index = CommandContexts.Num() - 1; Index >= 0; --Index)
 	{
 		delete CommandContexts[Index];
@@ -1735,3 +1727,31 @@ void FVulkanDevice::ForEachQueue(TFunctionRef<void(FVulkanQueue&)> Callback)
 		Callback(*Queues[QueueIndex]);
 	}
 }
+
+void FVulkanDevice::InitGlobalSamplers()
+{
+	checkf(SamplerMap.Num() == 0, TEXT("Global Samplers should be the first samplers created."));
+
+	GlobalSamplers[(uint32)FVulkanShaderHeader::EGlobalSamplerType::PointClampedSampler] =
+		ResourceCast(RHICreateSamplerState(FSamplerStateInitializerRHI(SF_Point, AM_Clamp, AM_Clamp, AM_Clamp)).GetReference());
+
+	GlobalSamplers[(uint32)FVulkanShaderHeader::EGlobalSamplerType::PointWrappedSampler] = 
+		ResourceCast(RHICreateSamplerState(FSamplerStateInitializerRHI(SF_Point, AM_Wrap, AM_Wrap, AM_Wrap)).GetReference());
+
+	GlobalSamplers[(uint32)FVulkanShaderHeader::EGlobalSamplerType::BilinearClampedSampler] =
+		ResourceCast(RHICreateSamplerState(FSamplerStateInitializerRHI(SF_Bilinear, AM_Clamp, AM_Clamp, AM_Clamp)).GetReference());
+
+	GlobalSamplers[(uint32)FVulkanShaderHeader::EGlobalSamplerType::BilinearWrappedSampler] =
+		ResourceCast(RHICreateSamplerState(FSamplerStateInitializerRHI(SF_Bilinear, AM_Wrap, AM_Wrap, AM_Wrap)).GetReference());
+
+	GlobalSamplers[(uint32)FVulkanShaderHeader::EGlobalSamplerType::TrilinearClampedSampler] =
+		ResourceCast(RHICreateSamplerState(FSamplerStateInitializerRHI(SF_Trilinear, AM_Clamp, AM_Clamp, AM_Clamp)).GetReference());
+
+	GlobalSamplers[(uint32)FVulkanShaderHeader::EGlobalSamplerType::TrilinearWrappedSampler] =
+		ResourceCast(RHICreateSamplerState(FSamplerStateInitializerRHI(SF_Trilinear, AM_Wrap, AM_Wrap, AM_Wrap)).GetReference());
+
+	static_assert((uint32)FVulkanShaderHeader::EGlobalSamplerType::Count == 6, "Need to add new global sampler decl!");
+
+	// Make sure if bindless handles are being generated that these samplers are the very first
+	check(!GlobalSamplers[0]->GetBindlessHandle().IsValid() || GlobalSamplers[0]->GetBindlessHandle().GetIndex() == 1);
+}
diff --git a/Engine/Source/Runtime/VulkanRHI/Private/VulkanDevice.h b/Engine/Source/Runtime/VulkanRHI/Private/VulkanDevice.h
--- a/Engine/Source/Runtime/VulkanRHI/Private/VulkanDevice.h
+++ b/Engine/Source/Runtime/VulkanRHI/Private/VulkanDevice.h
@@ -400,7 +400,7 @@ public:
 
 	const FVulkanSamplerState& GetDefaultSampler() const
 	{
-		return *DefaultSampler;
+		return GetGlobalSamplers(FVulkanShaderHeader::EGlobalSamplerType::PointWrappedSampler);
 	}
 
 	const VkFormatProperties& GetFormatProperties(VkFormat InFormat) const;
@@ -556,6 +556,11 @@ public:
 	FGPUTimingCalibrationTimestamp GetCalibrationTimestamp();
 #endif
 
+	const FVulkanSamplerState& GetGlobalSamplers(FVulkanShaderHeader::EGlobalSamplerType Type) const
+	{
+		return *GlobalSamplers[(uint32)Type];
+	}
+
 private:
 	void MapBufferFormatSupport(FPixelFormatInfo& PixelFormatInfo, EPixelFormat UEFormat, VkFormat VulkanFormat);
 	void MapImageFormatSupport(FPixelFormatInfo& PixelFormatInfo, const TArrayView<const VkFormat>& PrioritizedFormats, EPixelFormatCapabilities RequiredCapabilities);
@@ -564,6 +569,9 @@ private:
 	void MapFormatSupport(EPixelFormat UEFormat, std::initializer_list<VkFormat> PrioritizedFormats, const VkComponentMapping& ComponentMapping, int32 BlockBytes);
 	void MapFormatSupport(EPixelFormat UEFormat, std::initializer_list<VkFormat> PrioritizedFormats, const VkComponentMapping& ComponentMapping, EPixelFormatCapabilities RequiredCapabilities);
 
+	void InitGlobalSamplers();
+	TStaticArray<FVulkanSamplerState*, (uint32)FVulkanShaderHeader::EGlobalSamplerType::Count> GlobalSamplers;
+
 	VkDevice Device;
 
 	VulkanRHI::FDeviceMemoryManager DeviceMemoryManager;
@@ -591,8 +599,6 @@ private:
 
 	FVulkanShaderFactory ShaderFactory;
 
-	FVulkanSamplerState* DefaultSampler;
-
 	VkPhysicalDevice Gpu;
 	VkPhysicalDeviceProperties GpuProps;
 
diff --git a/Engine/Source/Runtime/VulkanRHI/Private/VulkanPipelineState.cpp b/Engine/Source/Runtime/VulkanRHI/Private/VulkanPipelineState.cpp
--- a/Engine/Source/Runtime/VulkanRHI/Private/VulkanPipelineState.cpp
+++ b/Engine/Source/Runtime/VulkanRHI/Private/VulkanPipelineState.cpp
@@ -56,6 +56,12 @@ FVulkanComputePipelineDescriptorState::FVulkanComputePipelineDescriptorState(FVu
 	CreateDescriptorWriteInfos();
 	InComputePipeline->AddRef();
 
+	for (const FVulkanShaderHeader::FGlobalSamplerInfo& GlobalSamplerInfo : CodeHeader.GlobalSamplerInfos)
+	{
+		checkSlow(!bUseBindless);
+		SetSamplerState(ShaderStage::Compute, GlobalSamplerInfo.BindingIndex, &Device->GetGlobalSamplers(GlobalSamplerInfo.Type));
+	}
+
 	ensure(DSWriter.Num() == 0 || DSWriter.Num() == 1);
 }
 
@@ -156,7 +162,6 @@ void FVulkanCommonPipelineDescriptorState::CreateDescriptorWriteInfos()
 	DescriptorSetHandles.AddZeroed(MaxNumSets);
 }
 
-
 static inline VulkanRHI::FVulkanAllocation UpdatePackedUniformBuffers(const FPackedUniformBuffers& PackedUniformBuffers, FVulkanCommandListContext* InContext, FVulkanCommandBuffer* InCmdBuffer)
 {
 	const FPackedUniformBuffers::FPackedBuffer& StagedUniformBuffer = PackedUniformBuffers.GetBuffer();
@@ -301,6 +306,7 @@ FVulkanGraphicsPipelineDescriptorState::FVulkanGraphicsPipelineDescriptorState(F
 
 	const FVulkanShaderFactory& ShaderFactory = Device->GetShaderFactory();
 
+	TStaticArray<const FVulkanShaderHeader*, ShaderStage::NumGraphicsStages> StageHeaders(InPlace, nullptr);
 
 	uint64 VertexShaderKey = InGfxPipeline->GetShaderKey(SF_Vertex);
 	if (VertexShaderKey)
@@ -310,6 +316,7 @@ FVulkanGraphicsPipelineDescriptorState::FVulkanGraphicsPipelineDescriptorState(F
 
 		PackedUniformBuffers[ShaderStage::Vertex].Init(VertexShader->GetCodeHeader(), PackedUniformBuffersMask[ShaderStage::Vertex]);
 		UsedSetsMask |= VertexShader->GetCodeHeader().Bindings.Num() ? (1u << ShaderStage::Vertex) : 0u;
+		StageHeaders[ShaderStage::Vertex] = &VertexShader->GetCodeHeader();
 	}
 
 	uint64 PixelShaderKey = InGfxPipeline->GetShaderKey(SF_Pixel);
@@ -320,6 +327,7 @@ FVulkanGraphicsPipelineDescriptorState::FVulkanGraphicsPipelineDescriptorState(F
 
 		PackedUniformBuffers[ShaderStage::Pixel].Init(PixelShader->GetCodeHeader(), PackedUniformBuffersMask[ShaderStage::Pixel]);
 		UsedSetsMask |= PixelShader->GetCodeHeader().Bindings.Num() ? (1u << ShaderStage::Pixel) : 0u;
+		StageHeaders[ShaderStage::Pixel] = &PixelShader->GetCodeHeader();
 	}
 
 #if PLATFORM_SUPPORTS_MESH_SHADERS
@@ -331,6 +339,7 @@ FVulkanGraphicsPipelineDescriptorState::FVulkanGraphicsPipelineDescriptorState(F
 
 		PackedUniformBuffers[ShaderStage::Mesh].Init(MeshShader->GetCodeHeader(), PackedUniformBuffersMask[ShaderStage::Mesh]);
 		UsedSetsMask |= MeshShader->GetCodeHeader().Bindings.Num() ? (1u << ShaderStage::Mesh) : 0u;
+		StageHeaders[ShaderStage::Mesh] = &MeshShader->GetCodeHeader();
 	}
 
 	uint64 TaskShaderKey = InGfxPipeline->GetShaderKey(SF_Amplification);
@@ -341,6 +350,7 @@ FVulkanGraphicsPipelineDescriptorState::FVulkanGraphicsPipelineDescriptorState(F
 
 		PackedUniformBuffers[ShaderStage::Task].Init(TaskShader->GetCodeHeader(), PackedUniformBuffersMask[ShaderStage::Task]);
 		UsedSetsMask |= TaskShader->GetCodeHeader().Bindings.Num() ? (1u << ShaderStage::Task) : 0u;
+		StageHeaders[ShaderStage::Task] = &TaskShader->GetCodeHeader();
 	}
 #endif
 
@@ -353,14 +363,27 @@ FVulkanGraphicsPipelineDescriptorState::FVulkanGraphicsPipelineDescriptorState(F
 
 		PackedUniformBuffers[ShaderStage::Geometry].Init(GeometryShader->GetCodeHeader(), PackedUniformBuffersMask[ShaderStage::Geometry]);
 		UsedSetsMask |= GeometryShader->GetCodeHeader().Bindings.Num() ? (1u << ShaderStage::Geometry) : 0u;
+		StageHeaders[ShaderStage::Geometry] = &GeometryShader->GetCodeHeader();
 	}
 #endif
 
 	CreateDescriptorWriteInfos();
-
 	//UE_LOG(LogVulkanRHI, Warning, TEXT("GfxPSOState %p For PSO %p Writes:%d"), this, InGfxPipeline, DSWriteContainer.DescriptorWrites.Num());
 
 	InGfxPipeline->AddRef();
+
+	for (int32 StageIndex = 0; StageIndex < StageHeaders.Num(); ++StageIndex)
+	{
+		const FVulkanShaderHeader* CodeHeader = StageHeaders[StageIndex];
+		if (CodeHeader)
+		{
+			for (const FVulkanShaderHeader::FGlobalSamplerInfo& GlobalSamplerInfo : CodeHeader->GlobalSamplerInfos)
+			{
+				checkSlow(!bUseBindless);
+				SetSamplerState(StageIndex, GlobalSamplerInfo.BindingIndex, &Device->GetGlobalSamplers(GlobalSamplerInfo.Type));
+			}
+		}
+	}
 }
 
 template<bool bUseDynamicGlobalUBs>
diff --git a/Engine/Source/Runtime/VulkanRHI/Public/VulkanShaderResources.h b/Engine/Source/Runtime/VulkanRHI/Public/VulkanShaderResources.h
--- a/Engine/Source/Runtime/VulkanRHI/Public/VulkanShaderResources.h
+++ b/Engine/Source/Runtime/VulkanRHI/Public/VulkanShaderResources.h
@@ -62,6 +62,28 @@ struct FVulkanShaderHeader
 	};
 	TArray<FInputAttachmentInfo> InputAttachmentInfos;
 
+	// Types of Global Samplers (see Common.ush for types)
+	// Must match GetGlobalSamplerType() in SpirVShaderCompiler.inl and
+	// and declarations in VulkanCommon.ush
+	enum class EGlobalSamplerType : uint8
+	{
+		PointClampedSampler,
+		PointWrappedSampler,
+		BilinearClampedSampler,
+		BilinearWrappedSampler,
+		TrilinearClampedSampler,
+		TrilinearWrappedSampler,
+
+		Count,
+		Invalid,
+	};
+	struct FGlobalSamplerInfo
+	{
+		uint8					BindingIndex;
+		EGlobalSamplerType		Type;
+	};
+	TArray<FGlobalSamplerInfo>  GlobalSamplerInfos;
+
 	// The number of uniform buffers containing constants and requiring bindings
 	// Uniform buffers beyond this index do not have bindings (resource only UB)
 	uint32						NumBoundUniformBuffers = 0;
@@ -138,11 +160,19 @@ inline FArchive& operator<<(FArchive& Ar, FVulkanShaderHeader::FInputAttachmentI
 	return Ar;
 }
 
+inline FArchive& operator<<(FArchive& Ar, FVulkanShaderHeader::FGlobalSamplerInfo& Info)
+{
+	Ar << Info.BindingIndex;
+	Ar << Info.Type;
+	return Ar;
+}
+
 inline FArchive& operator<<(FArchive& Ar, FVulkanShaderHeader& Header)
 {
 	Ar << Header.Bindings;
 	Ar << Header.UniformBufferInfos;
 	Ar << Header.InputAttachmentInfos;
+	Ar << Header.GlobalSamplerInfos;
 	Ar << Header.NumBoundUniformBuffers;
 	Ar << Header.PackedGlobalsSize;
 	Ar << Header.InputAttachmentsMask;
