Vulkan: Prevent unnecessary swapchain recreation when vsync is off.

#rnx

diff --git a/Engine/Source/Runtime/VulkanRHI/Private/VulkanViewport.cpp b/Engine/Source/Runtime/VulkanRHI/Private/VulkanViewport.cpp
--- a/Engine/Source/Runtime/VulkanRHI/Private/VulkanViewport.cpp
+++ b/Engine/Source/Runtime/VulkanRHI/Private/VulkanViewport.cpp
@@ -118,10 +118,13 @@ FVulkanViewport::FVulkanViewport(FVulkanDevice* InDevice, void* InWindowHandle,
 	, WindowHandle(InWindowHandle)
 	, PresentCount(0)
 	, bRenderOffscreen(false)
-	, LockToVsync(1)
 	, AcquiredSemaphore(nullptr)
 {
 	check(IsInGameThread());
+
+	static IConsoleVariable* CVarVsync = IConsoleManager::Get().FindConsoleVariable(TEXT("r.VSync"));
+	LockToVsync = CVarVsync->GetInt() != 0;
+
 	FVulkanDynamicRHI::Get().Viewports.Add(this);
 
 	// Make sure Instance is created
