VulkanRHI: Fix Vulkan calibration timestamps for new profiler in Linux to match values returned by Cycle64().

#rnx

diff --git a/Engine/Source/Runtime/VulkanRHI/Private/VulkanDevice.cpp b/Engine/Source/Runtime/VulkanRHI/Private/VulkanDevice.cpp
--- a/Engine/Source/Runtime/VulkanRHI/Private/VulkanDevice.cpp
+++ b/Engine/Source/Runtime/VulkanRHI/Private/VulkanDevice.cpp
@@ -249,6 +249,7 @@ FVulkanDevice::FVulkanDevice(FVulkanDynamicRHI* InRHI, VkPhysicalDevice InGpu)
 	UE_LOG(LogVulkanRHI, Display, TEXT("- API=%d.%d.%d (0x%x) Driver=0x%x VendorId=0x%x"), VK_VERSION_MAJOR(GpuProps.apiVersion), VK_VERSION_MINOR(GpuProps.apiVersion), VK_VERSION_PATCH(GpuProps.apiVersion), GpuProps.apiVersion, GpuProps.driverVersion, GpuProps.vendorID);
 	UE_LOG(LogVulkanRHI, Display, TEXT("- DeviceID=0x%x Type=%s"), GpuProps.deviceID, VK_TYPE_TO_STRING(VkPhysicalDeviceType, GpuProps.deviceType));
 	UE_LOG(LogVulkanRHI, Display, TEXT("- Max Descriptor Sets Bound %d, Timestamps %d"), GpuProps.limits.maxBoundDescriptorSets, GpuProps.limits.timestampComputeAndGraphics);
+	UE_LOG(LogVulkanRHI, Display, TEXT("- Calibrated Timestamp Domain: %s."), VK_TYPE_TO_STRING(VkTimeDomainKHR, FVulkanPlatform::GetTimeDomain()));
 
 	ensureMsgf(VendorId != EGpuVendorId::Unknown, TEXT("Unknown vendor ID 0x%x"), GpuProps.vendorID);
 }
@@ -1682,8 +1683,24 @@ void FVulkanDevice::GetCalibrationTimestamp(FVulkanTiming& InOutTiming)
 		VERIFYVULKANRESULT(VulkanDynamicAPI::vkGetCalibratedTimestampsEXT(GetInstanceHandle(), TimestampCount, CalibratedTimestampInfo, Timestamps, MaxDeviations));
 
 		InOutTiming.GPUTimestamp = Timestamps[0];
-		InOutTiming.CPUTimestamp = bHasHostTimeDomain ? Timestamps[1] : FPlatformTime::Cycles64();
-		check(InOutTiming.CPUTimestamp <= FPlatformTime::Cycles64());
+		switch (CalibratedTimestampInfo[1].timeDomain)
+		{
+		case VK_TIME_DOMAIN_CLOCK_MONOTONIC_KHR:
+			// Divide to match values returned by Cycles64()
+			InOutTiming.CPUTimestamp = Timestamps[1] / 100ULL;
+			break;
+
+		case VK_TIME_DOMAIN_QUERY_PERFORMANCE_COUNTER_KHR:
+			InOutTiming.CPUTimestamp = Timestamps[1];
+			break;
+
+		default:
+			InOutTiming.CPUTimestamp = FPlatformTime::Cycles64();
+			break;
+
+		}
+
+		checkf(InOutTiming.CPUTimestamp <= FPlatformTime::Cycles64(), TEXT("New calibration timestamp (%llu) ahead of current time (%llu)."), InOutTiming.CPUTimestamp, FPlatformTime::Cycles64());
 	}
 	else
 	{
