VulkanRHI: start timeline values at 0 and handle unused queues (fixes skipped signal value in new profiler).

#rnx

diff --git a/Engine/Source/Runtime/VulkanRHI/Private/VulkanQueue.h b/Engine/Source/Runtime/VulkanRHI/Private/VulkanQueue.h
--- a/Engine/Source/Runtime/VulkanRHI/Private/VulkanQueue.h
+++ b/Engine/Source/Runtime/VulkanRHI/Private/VulkanQueue.h
@@ -109,8 +109,8 @@ private:
 
 	const bool bUseTimelineSemaphores;
 	VulkanRHI::FSemaphore* TimelineSempahore = nullptr;
-	uint64 NextTimelineSemaphoreValue = 2;
-	uint64 CompletedTimelineSemaphoreValue = 1;
+	uint64 NextTimelineSemaphoreValue = 1;
+	uint64 CompletedTimelineSemaphoreValue = 0;
 
 	mutable FCriticalSection CS;
 	uint64 SubmitCounter;
diff --git a/Engine/Source/Runtime/VulkanRHI/Private/VulkanSubmission.cpp b/Engine/Source/Runtime/VulkanRHI/Private/VulkanSubmission.cpp
--- a/Engine/Source/Runtime/VulkanRHI/Private/VulkanSubmission.cpp
+++ b/Engine/Source/Runtime/VulkanRHI/Private/VulkanSubmission.cpp
@@ -121,7 +121,7 @@ void FVulkanDynamicRHI::InitializeSubmissionPipe()
 	{
 		if (CVarVulkanUseInterruptThread.GetValueOnAnyThread())
 		{
-			CPUTimelineSemaphore = new VulkanRHI::FSemaphore(*Device, VulkanRHI::EVulkanSemaphoreFlags::Timeline | VulkanRHI::EVulkanSemaphoreFlags::ImmediateDeletion, CPUTimelineSemaphoreValue);
+			CPUTimelineSemaphore = new VulkanRHI::FSemaphore(*Device, VulkanRHI::EVulkanSemaphoreFlags::Timeline | VulkanRHI::EVulkanSemaphoreFlags::ImmediateDeletion, CPUTimelineSemaphoreValue++);
 			InterruptThread = new FVulkanThread(TEXT("RHIInterruptThread"), TPri_Highest, this, &FVulkanDynamicRHI::WaitAndProcessInterruptQueue);
 		}
 
@@ -351,24 +351,32 @@ bool FVulkanDynamicRHI::WaitAndProcessInterruptQueue()
 				const int32 NumSemas = NumQueues + 1;  // add a timeline sema for CPU wake
 				VkSemaphore Semaphores[NumSemas];
 				uint64_t Values[NumSemas];
+				uint32 NumWaits = 0;
 				for (uint32 Index = 0; Index < NumSyncs; ++Index)
 				{
-					checkf(NextPayloads[Index]->TimelineSemaphoreValue > 0, TEXT("Payloads should all use the same types of syncs!"));
-					Semaphores[Index] = NextPayloads[Index]->Queue.GetTimelineSemaphore()->GetHandle();
-					Values[Index] = (uint64_t)NextPayloads[Index]->TimelineSemaphoreValue;
+					if (NextPayloads[Index]->TimelineSemaphoreValue > 0)
+					{
+						Semaphores[NumWaits] = NextPayloads[Index]->Queue.GetTimelineSemaphore()->GetHandle();
+						Values[NumWaits] = (uint64_t)NextPayloads[Index]->TimelineSemaphoreValue;
+						NumWaits++;
+					}
+					else
+					{
+						checkf(NextPayloads[Index]->CommandBuffers.Num() == 0, TEXT("TimelineSemaphoreValue should only be 0 on unused queues."));
+					}
 				}
 
 				if (CPUTimelineSemaphore)
 				{
-					Semaphores[NumSyncs] = CPUTimelineSemaphore->GetHandle();
-					Values[NumSyncs] = (uint64_t)(CPUTimelineSemaphoreValue);
-					NumSyncs++;
+					Semaphores[NumWaits] = CPUTimelineSemaphore->GetHandle();
+					Values[NumWaits] = (uint64_t)(CPUTimelineSemaphoreValue);
+					NumWaits++;
 				}
 
 				VkSemaphoreWaitInfo WaitInfo;
 				ZeroVulkanStruct(WaitInfo, VK_STRUCTURE_TYPE_SEMAPHORE_WAIT_INFO);
 				WaitInfo.flags = VK_SEMAPHORE_WAIT_ANY_BIT;
-				WaitInfo.semaphoreCount = NumSyncs;
+				WaitInfo.semaphoreCount = NumWaits;
 				WaitInfo.pSemaphores = Semaphores;
 				WaitInfo.pValues = Values;
 				VERIFYVULKANRESULT_EXPANDED(VulkanRHI::vkWaitSemaphoresKHR(Device->GetInstanceHandle(), &WaitInfo, (uint64_t)Timeout));
