VulkanRHI: Remove old async compute functions.  Replace single async compute context with an array.

#rnx

diff --git a/Engine/Config/Android/AndroidEngine.ini b/Engine/Config/Android/AndroidEngine.ini
--- a/Engine/Config/Android/AndroidEngine.ini
+++ b/Engine/Config/Android/AndroidEngine.ini
@@ -89,6 +89,8 @@ r.Vulkan.Submission.AllowTimelineSemaphores=0
 r.Vulkan.Submission.UseInterruptThread=0
 r.Vulkan.Submission.UseSubmissionThread=0
 r.Vulkan.RHIThread=1
+; Disable async compute until we can target devices to enable it on
+r.Vulkan.AllowAsyncCompute=0
 
 [/Script/Engine.RendererSettings]
 r.DistanceFields=1
diff --git a/Engine/Source/Runtime/VulkanRHI/Private/VulkanBarriers.cpp b/Engine/Source/Runtime/VulkanRHI/Private/VulkanBarriers.cpp
--- a/Engine/Source/Runtime/VulkanRHI/Private/VulkanBarriers.cpp
+++ b/Engine/Source/Runtime/VulkanRHI/Private/VulkanBarriers.cpp
@@ -1082,19 +1082,12 @@ private:
 	{
 #if DO_GUARD_SLOW
 		ERHIPipeline RequestedPipeline = IsBeginTransition ? Data->SrcPipelines : Data->DstPipelines;
-		const TCHAR* RequestedPipelineName = IsBeginTransition ? TEXT("SRC") : TEXT("DST");
-		switch (RequestedPipeline)
-		{
-		case ERHIPipeline::Graphics:
-			checkf(!Context.Device.IsRealAsyncComputeContext(&Context), TEXT("The %s pipeline for this transition is the Graphics queue, but it's submitted on the AsyncCompute queue."), IsBeginTransition ? TEXT("BEGIN") : TEXT("END"));
-			break;
-		case ERHIPipeline::AsyncCompute:
-			checkf(Context.Device.IsRealAsyncComputeContext(&Context), TEXT("The %s pipeline for this transition is the AsyncCompute queue, but it's submitted on the Graphics queue."), IsBeginTransition ? TEXT("BEGIN") : TEXT("END"));
-			break;
-		default:
-			checkNoEntry();
-			break;
-		}
+		checkf(EnumHasAnyFlags(Context.GetPipeline(), RequestedPipeline), 
+			TEXT("The %s pipelines for this transition are [%s], but it's submitted on the [%s] queue."), 
+			IsBeginTransition ? TEXT("SRC") : TEXT("DST"),
+			*GetRHIPipelineName(RequestedPipeline),
+			*GetRHIPipelineName(Context.GetPipeline())
+		);
 #endif
 	}
 
@@ -1394,8 +1387,7 @@ void FVulkanCommandListContext::RHIBeginTransitions(TArrayView<const FRHITransit
 
 void FVulkanCommandListContext::RHIEndTransitions(TArrayView<const FRHITransition*> Transitions)
 {
-	const ERHIPipeline CurrentPipeline = Device.IsRealAsyncComputeContext(this) ? ERHIPipeline::AsyncCompute : ERHIPipeline::Graphics;
-	TArray<VulkanRHI::FSemaphore*> WaitSemaphores = ExtractTransitionSemaphores(Transitions, CurrentPipeline);
+	TArray<VulkanRHI::FSemaphore*> WaitSemaphores = ExtractTransitionSemaphores(Transitions, GetPipeline());
 	if (WaitSemaphores.Num() > 0)
 	{
 		AddWaitSemaphores(VK_PIPELINE_STAGE_TOP_OF_PIPE_BIT, WaitSemaphores);
diff --git a/Engine/Source/Runtime/VulkanRHI/Private/VulkanCommands.cpp b/Engine/Source/Runtime/VulkanRHI/Private/VulkanCommands.cpp
--- a/Engine/Source/Runtime/VulkanRHI/Private/VulkanCommands.cpp
+++ b/Engine/Source/Runtime/VulkanRHI/Private/VulkanCommands.cpp
@@ -922,9 +922,7 @@ void FVulkanCommandListContext::RHIWriteGPUFence(FRHIGPUFence* FenceRHI)
 
 IRHIComputeContext* FVulkanDynamicRHI::RHIGetCommandContext(ERHIPipeline Pipeline, FRHIGPUMask GPUMask)
 {
-	// @todo: RHI command list refactor - fix async compute
-	checkf(Pipeline == ERHIPipeline::Graphics, TEXT("Async compute command contexts not currently implemented."));
-
-	FVulkanCommandListContext* CmdContext = Device->AcquireDeferredContext();
+	check(IsSingleRHIPipeline(Pipeline));
+	FVulkanCommandListContext* CmdContext = Device->AcquireDeferredContext(Pipeline);
 	return CmdContext;
 }
diff --git a/Engine/Source/Runtime/VulkanRHI/Private/VulkanContext.cpp b/Engine/Source/Runtime/VulkanRHI/Private/VulkanContext.cpp
--- a/Engine/Source/Runtime/VulkanRHI/Private/VulkanContext.cpp
+++ b/Engine/Source/Runtime/VulkanRHI/Private/VulkanContext.cpp
@@ -174,10 +174,11 @@ void FVulkanContextCommon::PrepareNewCommandBuffer(FVulkanPayload& Payload)
 }
 
 
-FVulkanCommandListContext::FVulkanCommandListContext(FVulkanDynamicRHI* InRHI, FVulkanDevice& InDevice, FVulkanQueue& InQueue, FVulkanCommandListContext* InImmediate)
-	: FVulkanContextCommon(InDevice, InQueue)
+FVulkanCommandListContext::FVulkanCommandListContext(FVulkanDynamicRHI* InRHI, FVulkanDevice& InDevice, ERHIPipeline InPipeline, FVulkanCommandListContext* InImmediate)
+	: FVulkanContextCommon(InDevice, *InDevice.GetQueue(InPipeline))
 	, RHI(InRHI)
 	, Immediate(InImmediate)
+	, RHIPipeline(InPipeline)
 #if (RHI_NEW_GPU_PROFILER == 0)
 	, GpuProfiler(this, &InDevice)
 #endif
diff --git a/Engine/Source/Runtime/VulkanRHI/Private/VulkanContext.h b/Engine/Source/Runtime/VulkanRHI/Private/VulkanContext.h
--- a/Engine/Source/Runtime/VulkanRHI/Private/VulkanContext.h
+++ b/Engine/Source/Runtime/VulkanRHI/Private/VulkanContext.h
@@ -217,7 +217,7 @@ private:
 class FVulkanCommandListContext : public FVulkanContextCommon, public IRHICommandContext
 {
 public:
-	FVulkanCommandListContext(FVulkanDynamicRHI* InRHI, FVulkanDevice& InDevice, FVulkanQueue& InQueue, FVulkanCommandListContext* InImmediate);
+	FVulkanCommandListContext(FVulkanDynamicRHI* InRHI, FVulkanDevice& InDevice, ERHIPipeline InPipeline, FVulkanCommandListContext* InImmediate);
 	virtual ~FVulkanCommandListContext();
 
 	static inline FVulkanCommandListContext& Get(FRHICommandListBase& RHICmdList)
@@ -230,6 +230,11 @@ public:
 		return Immediate == nullptr;
 	}
 
+	virtual ERHIPipeline GetPipeline() const override
+	{
+		return RHIPipeline;
+	}
+
 	template <class ShaderType> void SetResourcesFromTables(const ShaderType* RESTRICT);
 	void CommitGraphicsResourceTables();
 	void CommitComputeResourceTables();
@@ -379,6 +384,7 @@ public:
 protected:
 	FVulkanDynamicRHI* RHI;
 	FVulkanCommandListContext* Immediate;
+	const ERHIPipeline RHIPipeline;
 
 	void BeginOcclusionQueryBatch(uint32 NumQueriesInBatch);
 
@@ -437,7 +443,7 @@ public:
 		return static_cast<FVulkanCommandListContextImmediate&>(RHICmdList.GetContext().GetLowestLevelContext());
 	}
 
-	FVulkanCommandListContextImmediate(FVulkanDynamicRHI* InRHI, FVulkanDevice& InDevice, FVulkanQueue& InQueue);
+	FVulkanCommandListContextImmediate(FVulkanDynamicRHI* InRHI, FVulkanDevice& InDevice);
 };
 
 class FVulkanUploadContext : public FVulkanContextCommon, public IRHIUploadContext
diff --git a/Engine/Source/Runtime/VulkanRHI/Private/VulkanDevice.cpp b/Engine/Source/Runtime/VulkanRHI/Private/VulkanDevice.cpp
--- a/Engine/Source/Runtime/VulkanRHI/Private/VulkanDevice.cpp
+++ b/Engine/Source/Runtime/VulkanRHI/Private/VulkanDevice.cpp
@@ -24,7 +24,7 @@
 TAutoConsoleVariable<int32> GRHIAllowAsyncComputeCvar(
 	TEXT("r.Vulkan.AllowAsyncCompute"),
 	0,
-	TEXT("0 to disable async compute queue(if available)")
+	TEXT("0 to disable async compute queue (if available)")
 	TEXT("1 to allow async compute queue")
 );
 
@@ -221,7 +221,6 @@ FVulkanDevice::FVulkanDevice(FVulkanDynamicRHI* InRHI, VkPhysicalDevice InGpu)
 	, DeferredDeletionQueue(this)
 	, Gpu(InGpu)
 	, ImmediateContext(nullptr)
-	, ComputeContext(nullptr)
 	, PipelineStateCache(nullptr)
 {
 	RHI = InRHI;
@@ -1400,16 +1399,7 @@ void FVulkanDevice::InitGPU()
 	// always look in the saved directory (for the cache from previous run that wasn't moved over to stage directory)
 	CacheFilenames.Add(VulkanRHI::GetPipelineCacheFilename());
 
-	ImmediateContext = new FVulkanCommandListContextImmediate(GVulkanRHI, *this, *GetGraphicsQueue());
-
-	if (GetGraphicsQueue()->GetFamilyIndex() != GetComputeQueue()->GetFamilyIndex() && GRHIAllowAsyncComputeCvar.GetValueOnAnyThread() != 0)
-	{
-		ComputeContext = new FVulkanCommandListContextImmediate(GVulkanRHI, *this, *GetComputeQueue());
-	}
-	else
-	{
-		ComputeContext = ImmediateContext;
-	}
+	ImmediateContext = new FVulkanCommandListContextImmediate(GVulkanRHI, *this);
 
 	extern TAutoConsoleVariable<int32> GRHIThreadCvar;
 	if (GRHIThreadCvar->GetInt() > 1)
@@ -1417,8 +1407,8 @@ void FVulkanDevice::InitGPU()
 		int32 Num = FTaskGraphInterface::Get().GetNumWorkerThreads();
 		for (int32 Index = 0; Index < Num; Index++)
 		{
-			FVulkanCommandListContext* CmdContext = new FVulkanCommandListContext(GVulkanRHI, *this, *GetGraphicsQueue(), ImmediateContext);
-			CommandContexts.Add(CmdContext);
+			FVulkanCommandListContext* CmdContext = new FVulkanCommandListContext(GVulkanRHI, *this, ERHIPipeline::Graphics, ImmediateContext);
+			CommandContexts[ERHIPipeline::Graphics].Add(CmdContext);
 		}
 	}
 
@@ -1463,10 +1453,6 @@ void FVulkanDevice::Destroy()
 
 	// Release pending state that might hold references to RHI resources before we do final FlushPendingDeletes
 	ImmediateContext->ReleasePendingState();
-	if (ComputeContext && ComputeContext != ImmediateContext)
-	{
-		ComputeContext->ReleasePendingState();
-	}
 
 	if (TransientHeapCache)
 	{
@@ -1483,17 +1469,17 @@ void FVulkanDevice::Destroy()
 	delete DescriptorPoolsManager;
 	DescriptorPoolsManager = nullptr;
 
-	for (int32 Index = CommandContexts.Num() - 1; Index >= 0; --Index)
+	for (int32 PipelineIndex = 0; PipelineIndex < (int32)ERHIPipeline::Num; ++PipelineIndex)
 	{
-		delete CommandContexts[Index];
-	}
-	CommandContexts.Reset();
+		TArray<FVulkanCommandListContext*>& Contexts = CommandContexts[PipelineIndex];
 
-	if (ComputeContext != ImmediateContext)
-	{
-		delete ComputeContext;
+		for (int32 ContextIndex = Contexts.Num() - 1; ContextIndex >= 0; --ContextIndex)
+		{
+			delete Contexts[ContextIndex];
+		}
+
+		Contexts.Reset();
 	}
-	ComputeContext = nullptr;
 
 	delete ImmediateContext;
 	ImmediateContext = nullptr;
@@ -1584,11 +1570,6 @@ void FVulkanDevice::NotifyDeletedImage(VkImage Image, bool bRenderTarget)
 
 void FVulkanDevice::NotifyDeletedGfxPipeline(class FVulkanRHIGraphicsPipelineState* Pipeline)
 {
-	if (ComputeContext != ImmediateContext)
-	{
-		//ensure(0);
-	}
-
 	//#todo-rco: Loop through all contexts!
 	if (ImmediateContext && ImmediateContext->PendingGfxState)
 	{
@@ -1598,11 +1579,6 @@ void FVulkanDevice::NotifyDeletedGfxPipeline(class FVulkanRHIGraphicsPipelineSta
 
 void FVulkanDevice::NotifyDeletedComputePipeline(class FVulkanComputePipeline* Pipeline)
 {
-	if (ComputeContext && ComputeContext != ImmediateContext && ComputeContext->PendingComputeState)
-	{
-		ComputeContext->PendingComputeState->NotifyDeletedPipeline(Pipeline);
-	}
-
 	//#todo-rco: Loop through all contexts!
 	if (ImmediateContext && ImmediateContext->PendingComputeState)
 	{
@@ -1616,14 +1592,18 @@ void FVulkanDevice::NotifyDeletedComputePipeline(class FVulkanComputePipeline* P
 }
 
 static FCriticalSection GContextCS;
-FVulkanCommandListContext* FVulkanDevice::AcquireDeferredContext()
+FVulkanCommandListContext* FVulkanDevice::AcquireDeferredContext(ERHIPipeline Pipeline)
 {
 	FScopeLock Lock(&GContextCS);
-	if (CommandContexts.Num() == 0)
+
+	TArray<FVulkanCommandListContext*>& Contexts = CommandContexts[Pipeline];
+
+	if (Contexts.Num() == 0)
 	{
-		return new FVulkanCommandListContext(GVulkanRHI, *this, *GetGraphicsQueue(), ImmediateContext);
+		return new FVulkanCommandListContext(GVulkanRHI, *this, Pipeline, ImmediateContext);
 	}
-	return CommandContexts.Pop(EAllowShrinking::No);
+
+	return Contexts.Pop(EAllowShrinking::No);
 }
 
 void FVulkanDevice::ReleaseDeferredContext(FVulkanCommandListContext* InContext)
@@ -1631,7 +1611,8 @@ void FVulkanDevice::ReleaseDeferredContext(FVulkanCommandListContext* InContext)
 	check(InContext);
 	{
 		FScopeLock Lock(&GContextCS);
-		CommandContexts.Push(InContext);
+		TArray<FVulkanCommandListContext*>& Contexts = CommandContexts[InContext->GetPipeline()];
+		Contexts.Push(InContext);
 	}
 }
 
diff --git a/Engine/Source/Runtime/VulkanRHI/Private/VulkanDevice.h b/Engine/Source/Runtime/VulkanRHI/Private/VulkanDevice.h
--- a/Engine/Source/Runtime/VulkanRHI/Private/VulkanDevice.h
+++ b/Engine/Source/Runtime/VulkanRHI/Private/VulkanDevice.h
@@ -286,15 +286,18 @@ public:
 		return bPresentOnComputeQueue;
 	}
 
-	bool IsRealAsyncComputeContext(const FVulkanCommandListContext* InContext) const
+	FVulkanQueue* GetQueue(ERHIPipeline Pipeline)
 	{
-		if (bAsyncComputeQueue)
+		if (Pipeline == ERHIPipeline::Graphics) 
 		{
-			ensure((FVulkanCommandListContext*)ImmediateContext != ComputeContext);
-			return InContext == ComputeContext;
+			return GetQueue(EVulkanQueueType::Graphics); 
 		}
-	
-		return false;
+		else if (Pipeline == ERHIPipeline::AsyncCompute) 
+		{ 
+			return GetQueue(EVulkanQueueType::AsyncCompute); 
+		}
+		checkNoEntry();
+		return nullptr;
 	}
 
 	FVulkanQueue* GetQueue(EVulkanQueueType QueueType)
@@ -507,7 +510,7 @@ public:
 	void NotifyDeletedGfxPipeline(class FVulkanRHIGraphicsPipelineState* Pipeline);
 	void NotifyDeletedComputePipeline(class FVulkanComputePipeline* Pipeline);
 
-	FVulkanCommandListContext* AcquireDeferredContext();
+	FVulkanCommandListContext* AcquireDeferredContext(ERHIPipeline Pipeline);
 	void ReleaseDeferredContext(FVulkanCommandListContext* InContext);
 	void VulkanSetObjectName(VkObjectType Type, uint64_t Handle, const TCHAR* Name);
 	inline const FOptionalVulkanDeviceExtensions& GetOptionalExtensions() const
@@ -646,8 +649,7 @@ private:
 	TMap<uint32, FSamplerStateRHIRef> SamplerMap;
 
 	FVulkanCommandListContextImmediate* ImmediateContext;
-	FVulkanCommandListContext* ComputeContext;
-	TArray<FVulkanCommandListContext*> CommandContexts;
+	TRHIPipelineArray<TArray<FVulkanCommandListContext*>> CommandContexts;
 
 	FVulkanDynamicRHI* RHI = nullptr;
 	
diff --git a/Engine/Source/Runtime/VulkanRHI/Private/VulkanRHI.cpp b/Engine/Source/Runtime/VulkanRHI/Private/VulkanRHI.cpp
--- a/Engine/Source/Runtime/VulkanRHI/Private/VulkanRHI.cpp
+++ b/Engine/Source/Runtime/VulkanRHI/Private/VulkanRHI.cpp
@@ -421,8 +421,8 @@ FDynamicRHI* FVulkanDynamicRHIModule::CreateRHI(ERHIFeatureLevel::Type InRequest
 IMPLEMENT_MODULE(FVulkanDynamicRHIModule, VulkanRHI);
 
 
-FVulkanCommandListContextImmediate::FVulkanCommandListContextImmediate(FVulkanDynamicRHI* InRHI, FVulkanDevice& InDevice, FVulkanQueue& InQueue)
-	: FVulkanCommandListContext(InRHI, InDevice, InQueue, nullptr)
+FVulkanCommandListContextImmediate::FVulkanCommandListContextImmediate(FVulkanDynamicRHI* InRHI, FVulkanDevice& InDevice)
+	: FVulkanCommandListContext(InRHI, InDevice, ERHIPipeline::Graphics, nullptr)
 {
 }
 
@@ -911,7 +911,8 @@ void FVulkanDynamicRHI::InitInstance()
 		GSupportsParallelRenderingTasksWithSeparateRHIThread = GRHISupportsRHIThread ? FVulkanPlatform::SupportParallelRenderingTasks() : false;
 
 		//#todo-rco: Add newer Nvidia also
-		GSupportsEfficientAsyncCompute = (Device->ComputeContext != Device->ImmediateContext) && ((Device->GetVendorId() == EGpuVendorId::Amd) || FParse::Param(FCommandLine::Get(), TEXT("ForceAsyncCompute")));
+		GSupportsEfficientAsyncCompute = Device->HasAsyncComputeQueue();
+		UE_LOG(LogVulkanRHI, Display, TEXT("Vulkan Async Compute has been %s."), GSupportsEfficientAsyncCompute ? TEXT("ENABLED") : TEXT("DISABLED"));
 
 		GSupportsVolumeTextureRendering = FVulkanPlatform::SupportsVolumeTextureRendering();
 
