VulkanRHI: Enable Interrupt/Submission threads

#rnx

diff --git a/Engine/Source/Runtime/VulkanRHI/Private/VulkanSubmission.cpp b/Engine/Source/Runtime/VulkanRHI/Private/VulkanSubmission.cpp
--- a/Engine/Source/Runtime/VulkanRHI/Private/VulkanSubmission.cpp
+++ b/Engine/Source/Runtime/VulkanRHI/Private/VulkanSubmission.cpp
@@ -99,14 +99,14 @@ protected:
 
 static TAutoConsoleVariable<int32> CVarVulkanUseInterruptThread(
 	TEXT("r.Vulkan.Submission.UseInterruptThread"),
-	0,
+	1,
 	TEXT("  0: Process completed GPU work directly on the RHI thread.\n")
 	TEXT("  1: Create a dedicated thread to process completed GPU work.\n"),
 	ECVF_ReadOnly);
 
 static TAutoConsoleVariable<int32> CVarVulkanUseSubmissionThread(
 	TEXT("r.Vulkan.Submission.UseSubmissionThread"),
-	0,
+	1,
 	TEXT("  0: Submit GPU work directly from the RHI thread.\n")
 	TEXT("  1: Create a dedicated thread to submit GPU work.\n"),
 	ECVF_ReadOnly);
