VulkanRHI: Add check for queue type on graphic context calls.

#rnx

diff --git a/Engine/Source/Runtime/VulkanRHI/Private/VulkanCommands.cpp b/Engine/Source/Runtime/VulkanRHI/Private/VulkanCommands.cpp
--- a/Engine/Source/Runtime/VulkanRHI/Private/VulkanCommands.cpp
+++ b/Engine/Source/Runtime/VulkanRHI/Private/VulkanCommands.cpp
@@ -225,6 +225,7 @@ void FVulkanCommandListContext::SetResourcesFromTables(const ShaderType* Shader)
 
 void FVulkanCommandListContext::CommitGraphicsResourceTables()
 {
+	checkfSlow(Queue.GetQueueType() == EVulkanQueueType::Graphics, TEXT("Recording a graphic command on a non-graphic queue."));
 	checkSlow(PendingGfxState);
 
 	if (const FVulkanShader* Shader = PendingGfxState->GetCurrentShader(SF_Vertex))
diff --git a/Engine/Source/Runtime/VulkanRHI/Private/VulkanQueue.h b/Engine/Source/Runtime/VulkanRHI/Private/VulkanQueue.h
--- a/Engine/Source/Runtime/VulkanRHI/Private/VulkanQueue.h
+++ b/Engine/Source/Runtime/VulkanRHI/Private/VulkanQueue.h
@@ -41,6 +41,11 @@ public:
 	FVulkanQueue(FVulkanDevice& InDevice, uint32 InFamilyIndex, EVulkanQueueType InQueueType);
 	~FVulkanQueue();
 
+	EVulkanQueueType GetQueueType()
+	{
+		return QueueType;
+	}
+
 	inline uint32 GetFamilyIndex() const
 	{
 		return FamilyIndex;
