VulkanRHI: Fix new profiler stomp in Vulkan

#rnx

diff --git a/Engine/Source/Runtime/VulkanRHI/Private/VulkanCommandBuffer.h b/Engine/Source/Runtime/VulkanRHI/Private/VulkanCommandBuffer.h
--- a/Engine/Source/Runtime/VulkanRHI/Private/VulkanCommandBuffer.h
+++ b/Engine/Source/Runtime/VulkanRHI/Private/VulkanCommandBuffer.h
@@ -126,8 +126,10 @@ public:
 	{
 		for (UE::RHI::GPUProfiler::FEvent::FBeginWork* BeginEvent : BeginEvents)
 		{
+			checkSlow((BeginEvent->CPUTimestamp == 0) && (BeginEvent->GPUTimestampTOP == UINT64_MAX));
 			BeginEvent->CPUTimestamp = CPUTimestamp;
 		}
+		BeginEvents.Reset();
 
 		Destination.Append(MoveTemp(EventStream));
 	}
