VulkanRHI: Remove unused code

#rnx

diff --git a/Engine/Source/Runtime/VulkanRHI/Private/VulkanCommandBuffer.cpp b/Engine/Source/Runtime/VulkanRHI/Private/VulkanCommandBuffer.cpp
--- a/Engine/Source/Runtime/VulkanRHI/Private/VulkanCommandBuffer.cpp
+++ b/Engine/Source/Runtime/VulkanRHI/Private/VulkanCommandBuffer.cpp
@@ -10,19 +10,9 @@
 #include "VulkanMemory.h"
 #include "VulkanRayTracing.h"
 
-static int32 GUseSingleQueue = 0;
-static FAutoConsoleVariableRef CVarVulkanUseSingleQueue(
-	TEXT("r.Vulkan.UseSingleQueue"),
-	GUseSingleQueue,
-	TEXT("Forces using the same queue for uploads and graphics.\n")
-	TEXT(" 0: Uses multiple queues(default)\n")
-	TEXT(" 1: Always uses the gfx queue for submissions"),
-	ECVF_Default
-);
 
 #define CMD_BUFFER_TIME_TO_WAIT_BEFORE_DELETING		10
 
-const uint32 GNumberOfFramesBeforeDeletingDescriptorPool = 300;
 
 FVulkanCommandBuffer::FVulkanCommandBuffer(FVulkanDevice& InDevice, FVulkanCommandBufferPool* InCommandBufferPool)
 	: Device(InDevice)
