VulkanRHI: Hold pool lock before changing command buffer state.

#rnx

diff --git a/Engine/Source/Runtime/VulkanRHI/Private/VulkanCommandBuffer.cpp b/Engine/Source/Runtime/VulkanRHI/Private/VulkanCommandBuffer.cpp
--- a/Engine/Source/Runtime/VulkanRHI/Private/VulkanCommandBuffer.cpp
+++ b/Engine/Source/Runtime/VulkanRHI/Private/VulkanCommandBuffer.cpp
@@ -218,6 +218,7 @@ bool FVulkanCommandBuffer::AcquirePoolSetAndDescriptorsIfNeeded(const class FVul
 
 void FVulkanCommandBuffer::ResetFenceStatus()
 {
+	FScopeLock ScopeLock(CommandBufferPool->GetCS());
 	if (State == EState::Submitted)
 	{
 		bHasPipeline = false;
@@ -246,6 +247,19 @@ void FVulkanCommandBuffer::ResetFenceStatus()
 	}
 }
 
+void FVulkanCommandBuffer::SetSubmitted()
+{
+	FScopeLock Lock(CommandBufferPool->GetCS());
+	State = FVulkanCommandBuffer::EState::Submitted;
+	SubmittedTime = FPlatformTime::Seconds();
+
+#if RHI_NEW_GPU_PROFILER
+	CommandBuffer->FlushProfilerEvents(Payload->EventStream, SubmitTime);
+#endif // RHI_NEW_GPU_PROFILER
+}
+
+
+
 FVulkanCommandBufferPool::FVulkanCommandBufferPool(FVulkanDevice& InDevice, FVulkanContextCommon& InContext)
 	: Device(InDevice)
 	, Context(InContext)
diff --git a/Engine/Source/Runtime/VulkanRHI/Private/VulkanCommandBuffer.h b/Engine/Source/Runtime/VulkanRHI/Private/VulkanCommandBuffer.h
--- a/Engine/Source/Runtime/VulkanRHI/Private/VulkanCommandBuffer.h
+++ b/Engine/Source/Runtime/VulkanRHI/Private/VulkanCommandBuffer.h
@@ -74,6 +74,8 @@ public:
 		return CommandBufferHandle;
 	}
 
+	void SetSubmitted();
+
 	void Begin(FVulkanQueryPool* QueryPool);
 	void End(FVulkanQueryPool* QueryPool);
 
diff --git a/Engine/Source/Runtime/VulkanRHI/Private/VulkanQueue.cpp b/Engine/Source/Runtime/VulkanRHI/Private/VulkanQueue.cpp
--- a/Engine/Source/Runtime/VulkanRHI/Private/VulkanQueue.cpp
+++ b/Engine/Source/Runtime/VulkanRHI/Private/VulkanQueue.cpp
@@ -243,13 +243,7 @@ void FVulkanQueue::SubmitPayloads(TArrayView<FVulkanPayload*> Payloads, TMap<VkS
 			for (FVulkanCommandBuffer* CommandBuffer : Payload->CommandBuffers)
 			{
 				CommandBufferStorage.Add(CommandBuffer->GetHandle());
-
-				CommandBuffer->State = FVulkanCommandBuffer::EState::Submitted;
-				CommandBuffer->SubmittedTime = FPlatformTime::Seconds();
-
-#if RHI_NEW_GPU_PROFILER
-				CommandBuffer->FlushProfilerEvents(Payload->EventStream, SubmitTime);
-#endif // RHI_NEW_GPU_PROFILER
+				CommandBuffer->SetSubmitted();
 			}
 			SubmitInfo.commandBufferCount = Payload->CommandBuffers.Num();
 			SubmitInfo.pCommandBuffers = &CommandBufferStorage[FirstIndex];
