Fixed a compilation error in FVulkanCommandBuffer::SetSubmitted when enabling the RHI_NEW_GPU_PROFILER macro. Some of the referenced symbols in that function have been refactored and needed to be updated.

diff --git a/Engine/Source/Runtime/VulkanRHI/Private/VulkanCommandBuffer.cpp b/Engine/Source/Runtime/VulkanRHI/Private/VulkanCommandBuffer.cpp
--- a/Engine/Source/Runtime/VulkanRHI/Private/VulkanCommandBuffer.cpp
+++ b/Engine/Source/Runtime/VulkanRHI/Private/VulkanCommandBuffer.cpp
@@ -254,12 +254,10 @@ void FVulkanCommandBuffer::SetSubmitted()
 	SubmittedTime = FPlatformTime::Seconds();
 
 #if RHI_NEW_GPU_PROFILER
-	CommandBuffer->FlushProfilerEvents(Payload->EventStream, SubmitTime);
+	FlushProfilerEvents(EventStream, SubmittedTime);
 #endif // RHI_NEW_GPU_PROFILER
 }
 
-
-
 FVulkanCommandBufferPool::FVulkanCommandBufferPool(FVulkanDevice& InDevice, FVulkanContextCommon& InContext)
 	: Device(InDevice)
 	, Context(InContext)
