VulkanRHI: Fix new profiler compiler issues in Vulkan on Android.

#rnx

diff --git a/Engine/Source/Runtime/VulkanRHI/Private/VulkanQueue.cpp b/Engine/Source/Runtime/VulkanRHI/Private/VulkanQueue.cpp
--- a/Engine/Source/Runtime/VulkanRHI/Private/VulkanQueue.cpp
+++ b/Engine/Source/Runtime/VulkanRHI/Private/VulkanQueue.cpp
@@ -93,21 +93,21 @@ int32 FVulkanQueue::SubmitQueuedPayloads(TMap<VkSemaphore, FBinarySemaphoreSigna
 				const VkSemaphore SemaphoreHandle = WaitSema->GetHandle();
 				FBinarySemaphoreSignalInfo& SignalInfo = SignaledSemas[SemaphoreHandle];
 
-				const int32 QueueType = (int32)SignalInfo.QueueType;
-				MaxTimelineWaitValue[QueueType] = FMath::Max(MaxTimelineWaitValue[QueueType], SignalInfo.TimelineValue);
+				const int32 QueueTypeIndex = (int32)SignalInfo.QueueType;
+				MaxTimelineWaitValue[QueueTypeIndex] = FMath::Max(MaxTimelineWaitValue[QueueTypeIndex], SignalInfo.TimelineValue);
 			}
 		}
 
 		// Create the profiler events
 		const uint64 SubmitTime = FPlatformTime::Cycles64();
-		for (int32 QueueType = 0; QueueType < (int32)EVulkanQueueType::Count; ++QueueType)
+		for (int32 QueueTypeIndex = 0; QueueTypeIndex < (int32)EVulkanQueueType::Count; ++QueueTypeIndex)
 		{
-			if (MaxTimelineWaitValue[QueueType] > 0)
+			if (MaxTimelineWaitValue[QueueTypeIndex] > 0)
 			{
-				FVulkanQueue* SignalQueue = Device.GetQueue((EVulkanQueueType)QueueType);
+				FVulkanQueue* SignalQueue = Device.GetQueue((EVulkanQueueType)QueueTypeIndex);
 				check(SignalQueue);
 				Payload->EventStream.Emplace<UE::RHI::GPUProfiler::FEvent::FWaitFence>(
-					SubmitTime, MaxTimelineWaitValue[QueueType], SignalQueue->GetProfilerQueue());
+					SubmitTime, MaxTimelineWaitValue[QueueTypeIndex], SignalQueue->GetProfilerQueue());
 			}
 		}
 #endif
