Fix for Vulkan PSO cache chunk release

diff --git a/Engine/Source/Runtime/VulkanRHI/Private/VulkanChunkedPipelineCache.cpp b/Engine/Source/Runtime/VulkanRHI/Private/VulkanChunkedPipelineCache.cpp
--- a/Engine/Source/Runtime/VulkanRHI/Private/VulkanChunkedPipelineCache.cpp
+++ b/Engine/Source/Runtime/VulkanRHI/Private/VulkanChunkedPipelineCache.cpp
@@ -1165,9 +1165,9 @@ private:
 
 		// unload oldest
 		FVulkanPipelineCacheChunkLRUNode* CurrentNode = CacheChunkLRU.GetTail();
-		for (int32 UnloadCount = 0; CurrentNode && UnloadCount < MaxToUnloadPerTick && CurrentResidentSize > TargetResidentCacheSizeBytes; UnloadCount++)
+		for (int32 UnloadCount = 0; CurrentNode && UnloadCount < MaxToUnloadPerTick && CurrentResidentSize > TargetResidentCacheSizeBytes;)
 		{
-			FVulkanPipelineCacheChunkLRUNode* NextNode = CurrentNode->GetNextNode();
+			FVulkanPipelineCacheChunkLRUNode* NextNode = CurrentNode->GetPrevNode();
 
 			UE::Vulkan::FPipelineCacheChunkKey ChunkKey = CurrentNode->GetValue();
 			FVulkanPipelineCacheChunk* FoundCacheChunk = CacheChunksMap.FindChecked(ChunkKey).Get();
