VulkanRHI: Fix command buffer event stream for new profiler

#rnx

diff --git a/Engine/Source/Runtime/VulkanRHI/Private/VulkanCommandBuffer.cpp b/Engine/Source/Runtime/VulkanRHI/Private/VulkanCommandBuffer.cpp
--- a/Engine/Source/Runtime/VulkanRHI/Private/VulkanCommandBuffer.cpp
+++ b/Engine/Source/Runtime/VulkanRHI/Private/VulkanCommandBuffer.cpp
@@ -252,10 +252,6 @@ void FVulkanCommandBuffer::SetSubmitted()
 	FScopeLock Lock(CommandBufferPool->GetCS());
 	State = FVulkanCommandBuffer::EState::Submitted;
 	SubmittedTime = FPlatformTime::Seconds();
-
-#if RHI_NEW_GPU_PROFILER
-	FlushProfilerEvents(EventStream, SubmittedTime);
-#endif // RHI_NEW_GPU_PROFILER
 }
 
 FVulkanCommandBufferPool::FVulkanCommandBufferPool(FVulkanDevice& InDevice, FVulkanContextCommon& InContext)
diff --git a/Engine/Source/Runtime/VulkanRHI/Private/VulkanQueue.cpp b/Engine/Source/Runtime/VulkanRHI/Private/VulkanQueue.cpp
--- a/Engine/Source/Runtime/VulkanRHI/Private/VulkanQueue.cpp
+++ b/Engine/Source/Runtime/VulkanRHI/Private/VulkanQueue.cpp
@@ -243,6 +243,9 @@ void FVulkanQueue::SubmitPayloads(TArrayView<FVulkanPayload*> Payloads, TMap<VkS
 			for (FVulkanCommandBuffer* CommandBuffer : Payload->CommandBuffers)
 			{
 				CommandBufferStorage.Add(CommandBuffer->GetHandle());
+#if RHI_NEW_GPU_PROFILER
+				CommandBuffer->FlushProfilerEvents(Payload->EventStream, SubmitTime);
+#endif // RHI_NEW_GPU_PROFILER
 				CommandBuffer->SetSubmitted();
 			}
 			SubmitInfo.commandBufferCount = Payload->CommandBuffers.Num();
