Fixed resource leak in texture from HardwareBuffer
#rnx

diff --git a/Engine/Source/Runtime/VulkanRHI/Private/VulkanTexture.cpp b/Engine/Source/Runtime/VulkanRHI/Private/VulkanTexture.cpp
--- a/Engine/Source/Runtime/VulkanRHI/Private/VulkanTexture.cpp
+++ b/Engine/Source/Runtime/VulkanRHI/Private/VulkanTexture.cpp
@@ -1810,7 +1810,14 @@ FVulkanTexture::FVulkanTexture(FVulkanDevice& InDevice, const FRHITextureCreateD
 	VkSamplerYcbcrConversion SamplerYcbcrConversion;
 	VERIFYVULKANRESULT(VulkanRHI::vkCreateSamplerYcbcrConversion(VulkanDevice, &SamplerYcbcrConversionCreateInfo, AllocationCallbacks, &SamplerYcbcrConversion));
 
-	FVulkanAndroidTextureResources* VulkanAndroidTextureResources = new FVulkanAndroidTextureResources{ VulkanImage, VulkanDeviceMemory, SamplerYcbcrConversion };
+	FVulkanAndroidTextureResources* VulkanAndroidTextureResources = new FVulkanAndroidTextureResources
+	{
+		VulkanImage,
+		VulkanDeviceMemory,
+		SamplerYcbcrConversion,
+		HardwareBuffer
+	};
+
 	ExternalImageDeleteCallbackInfo =
 	{
 		VulkanAndroidTextureResources,
