Android Window refactor
Can be enabled via configrules 'EnableNewBackgroundBehavior=true'
- Android native app window events no longer block and hold the RHI thread for the entire duration of backgrounding.
- Offscreen rendering fallback is used when no window is available.
- This prevents several classes of ANR and also resolves some scenarios involving GT time outs while waiting for RT. (hang detector, GC etc)
- FAndroidWindow now provides an accessor to the android ANativeWindow handle, the handle is locked and cannot be modified when acquired.
- RHI Viewport code now uses the ANativeWindow passed through via FAndroidWindow.
- Allows for easier implementation of multiwindow on android.
- Window loss/creation no longer affect android's pause/resume state.

diff --git a/Engine/Source/Runtime/ApplicationCore/Private/Android/AndroidWindow.cpp b/Engine/Source/Runtime/ApplicationCore/Private/Android/AndroidWindow.cpp
--- a/Engine/Source/Runtime/ApplicationCore/Private/Android/AndroidWindow.cpp
+++ b/Engine/Source/Runtime/ApplicationCore/Private/Android/AndroidWindow.cpp
@@ -3,6 +3,7 @@
 #include "Android/AndroidWindow.h"
 #include "Android/AndroidWindowUtils.h"
 #include "Android/AndroidEventManager.h"
+#include "Android/AndroidDynamicRHI.h"
 
 #if USE_ANDROID_JNI
 #include <android/native_window.h>
@@ -106,13 +107,31 @@ FAndroidWindow::~FAndroidWindow()
 
 TSharedRef<FAndroidWindow> FAndroidWindow::Make()
 {
-	return MakeShareable( new FAndroidWindow() );
+	TSharedRef<FAndroidWindow> Window = MakeShareable( new FAndroidWindow() );
+	int Num = FAndroidWindowManager::Get().AddWindow(Window);
+
+	if( Num == 1)  // special case for native activity, the first window gets any pre-existing ANative window. if there is not yet a pending window then the event thread will set it later.
+	{
+		// lock both the FAndroidWindow and the pending window to transfer the pending window across
+		FNativeAccessor MyANativeAccessor = Window->GetANativeAccessor(true);
+		FAndroidWindowManager::FPendingWindowAccessor PendingAccessor = FAndroidWindowManager::GetPendingWindowAccessor();
+
+		MyANativeAccessor.SetANativeWindow(PendingAccessor.GetANativeWindow());
+		PendingAccessor.SetANativeWindow(nullptr);
+	}
+
+	return Window;
 }
 
 FAndroidWindow::FAndroidWindow()
 {
 }
 
+void FAndroidWindow::Destroy()
+{
+	FAndroidWindowManager::Get().RemoveWindow(SharedThis(this));
+}
+
 void FAndroidWindow::Initialize( class FAndroidApplication* const Application, const TSharedRef< FGenericWindowDefinition >& InDefinition, const TSharedPtr< FAndroidWindow >& InParent, const bool bShowImmediately )
 {
 	OwningApplication = Application;
diff --git a/Engine/Source/Runtime/ApplicationCore/Public/Android/AndroidApplication.h b/Engine/Source/Runtime/ApplicationCore/Public/Android/AndroidApplication.h
--- a/Engine/Source/Runtime/ApplicationCore/Public/Android/AndroidApplication.h
+++ b/Engine/Source/Runtime/ApplicationCore/Public/Android/AndroidApplication.h
@@ -14,7 +14,7 @@ namespace FAndroidAppEntry
 	void PlatformInit();
 
 	// if the native window handle has changed then the new handle is required.
-	void ReInitWindow(void* NewNativeWindowHandle = nullptr);
+	void ReInitWindow(const TOptional<FAndroidWindow::FNativeAccessor>& WindowContainer);
 
 	void ReleaseEGL();
 	void OnPauseEvent();
diff --git a/Engine/Source/Runtime/ApplicationCore/Public/Android/AndroidWindow.h b/Engine/Source/Runtime/ApplicationCore/Public/Android/AndroidWindow.h
--- a/Engine/Source/Runtime/ApplicationCore/Public/Android/AndroidWindow.h
+++ b/Engine/Source/Runtime/ApplicationCore/Public/Android/AndroidWindow.h
@@ -9,14 +9,16 @@
 #if USE_ANDROID_JNI
 #include <android/native_window_jni.h>
 #endif
-
+#include "Async/SharedMutex.h"
+#include "Async/Mutex.h"
+#include "Async/UniqueLock.h"
 
 /**
  * A platform specific implementation of FNativeWindow.
  * Native windows provide platform-specific backing for and are always owned by an SWindow.
  */
  
-class FAndroidWindow : public FGenericWindow
+class FAndroidWindow : public FGenericWindow, public TSharedFromThis<FAndroidWindow>
 {
 public:
 	APPLICATIONCORE_API ~FAndroidWindow();
@@ -28,8 +30,9 @@ public:
 	 */
 	static APPLICATIONCORE_API TSharedRef<FAndroidWindow> Make();
 
+	APPLICATIONCORE_API virtual void Destroy() override;
 	
-	virtual void* GetOSWindowHandle() const override { return nullptr; }
+	virtual void* GetOSWindowHandle() const override { return (void*) this; }
 
 	APPLICATIONCORE_API void Initialize( class FAndroidApplication* const Application, const TSharedRef< FGenericWindowDefinition >& InDefinition, const TSharedPtr< FAndroidWindow >& InParent, const bool bShowImmediately );
 
@@ -71,11 +74,77 @@ public:
 
 	APPLICATIONCORE_API bool GetNativeWindowResolution(int32_t& OutWidth, int32_t& OutHeight) const;
 
+	class FNativeAccessor
+	{
+		bool bWriteAccess;
+		class FAndroidWindow* ProtectedObj = nullptr;
+		FNativeAccessor() {}
+	public:
+		FNativeAccessor(FNativeAccessor&& obj)
+		{
+			ProtectedObj = obj.ProtectedObj;
+			bWriteAccess = obj.bWriteAccess;
+			obj.ProtectedObj = nullptr;
+		}
+
+		FNativeAccessor(const FNativeAccessor&) = delete;
+		FNativeAccessor(bool bWriteAccessIn, FAndroidWindow* ProtectedObjIn) : bWriteAccess(bWriteAccessIn), ProtectedObj(ProtectedObjIn) 
+		{ 
+			if (bWriteAccess)
+			{
+				ProtectedObj->ANativeHandleMutex.Lock();
+			}
+			else
+			{
+				ProtectedObj->ANativeHandleMutex.LockShared();
+			}
+		}
+
+		~FNativeAccessor()
+		{ 
+			if(ProtectedObj)
+			{
+				if (bWriteAccess)
+				{
+					ProtectedObj->ANativeHandleMutex.Unlock();
+				}
+				else
+				{
+					ProtectedObj->ANativeHandleMutex.UnlockShared();
+				}
+			}
+		}
+
+		FAndroidWindow* Get() { return ProtectedObj; }
+		const FAndroidWindow* Get() const { return ProtectedObj; }
+
+		FAndroidWindow* operator ->() { return Get(); }
+		const FAndroidWindow* operator ->() const { return Get(); }
+
+		void SetANativeWindow(ANativeWindow* ANativeHandle)
+		{ 
+			check(bWriteAccess); 
+			ProtectedObj->PreviousANativeHandle = ProtectedObj->CurrentANativeHandle;
+			ProtectedObj->CurrentANativeHandle = ANativeHandle; 
+		}
+		ANativeWindow* GetANativeWindow() const { return ProtectedObj->CurrentANativeHandle; }
+		ANativeWindow* GetPreviousANativeWindow() const { return ProtectedObj->PreviousANativeHandle; }
+	};
+
+	FNativeAccessor GetANativeAccessor(bool bWriteAccess)
+	{
+		return FNativeAccessor(bWriteAccess, this);
+	}
+
 protected:
 	/** @return true if the native window is currently in fullscreen mode, false otherwise */
 	virtual EWindowMode::Type GetWindowMode() const override { return EWindowMode::Fullscreen; }
 
 private:
+	ANativeWindow* CurrentANativeHandle = nullptr;
+	ANativeWindow* PreviousANativeHandle = nullptr;
+	UE::FSharedMutex ANativeHandleMutex;
+
 	/**
 	 * Protect the constructor; only TSharedRefs of this class can be made.
 	 */
@@ -95,4 +164,86 @@ private:
 	static APPLICATIONCORE_API bool bAreCachedNativeDimensionsValid;
 	static APPLICATIONCORE_API int32 CachedNativeWindowWidth;
 	static APPLICATIONCORE_API int32 CachedNativeWindowHeight;
+
+	friend class FAndroidWindowManager;
+};
+
+class FAndroidWindowManager
+{
+	friend class FAndroidWindow;
+
+public:
+
+	static FAndroidWindowManager& Get()
+	{
+		static FAndroidWindowManager Single;
+		return Single;
+	}
+
+	TOptional<FAndroidWindow::FNativeAccessor> GetMainWindowAsNativeAccessor(bool bWriteAccess)
+	{
+		UE::TUniqueLock Lock(WindowManagerLock);
+		if (Windows.IsEmpty())
+		{
+			return TOptional<FAndroidWindow::FNativeAccessor>();
+		}
+
+		FAndroidWindow::FNativeAccessor Accessor = Windows[0]->GetANativeAccessor(bWriteAccess);
+		return TOptional<FAndroidWindow::FNativeAccessor>(MoveTemp(Accessor));
+	}
+
+
+	TOptional<FAndroidWindow::FNativeAccessor > FindFromANativeWindow(bool bWriteAccess, bool bSearchPrevious, const void* SearchANativeHandle)
+	{
+		UE::TUniqueLock Lock(WindowManagerLock);
+		for (TSharedRef<FAndroidWindow>& Window : Windows)
+		{
+			FAndroidWindow::FNativeAccessor Accessor = Window->GetANativeAccessor(bWriteAccess);
+			if (SearchANativeHandle == ( bSearchPrevious ? Accessor.GetPreviousANativeWindow() : Accessor.GetANativeWindow() ) )
+			{
+				return TOptional<FAndroidWindow::FNativeAccessor>(MoveTemp(Accessor));
+			}
+		}
+
+		return TOptional<FAndroidWindow::FNativeAccessor>();
+	}
+
+	// Android native window handling.
+	// this is used to collect the native window before the game has created its own FAndroidWindow.
+	class FPendingWindowAccessor
+	{
+		static inline ANativeWindow* PendingANativeWindow = nullptr;
+		static inline UE::FMutex Lock;
+		FPendingWindowAccessor(const FPendingWindowAccessor&) = delete;
+		FPendingWindowAccessor(FPendingWindowAccessor&&) = delete;
+	public:
+		FPendingWindowAccessor() { Lock.Lock(); }
+		~FPendingWindowAccessor() { Lock.Unlock(); }
+
+		void SetANativeWindow(ANativeWindow* ANativeHandle) { PendingANativeWindow = ANativeHandle; }
+		ANativeWindow* GetANativeWindow() const { return PendingANativeWindow; }
+	};
+
+	static FPendingWindowAccessor GetPendingWindowAccessor()
+	{
+		return FPendingWindowAccessor();
+	}
+
+private:
+	TArray< TSharedRef<FAndroidWindow>> Windows;
+
+	int AddWindow(TSharedRef<FAndroidWindow> AddMe)
+	{
+		UE::TUniqueLock Lock(WindowManagerLock);
+		Windows.Add(AddMe);
+		return Windows.Num();
+	}
+	
+	void RemoveWindow(TSharedRef<FAndroidWindow> RemoveMe)
+	{
+		UE::TUniqueLock Lock(WindowManagerLock);
+		Windows.Remove(RemoveMe);
+	}
+
+	mutable UE::FMutex WindowManagerLock;
 };
diff --git a/Engine/Source/Runtime/Core/Private/Android/AndroidPlatformMisc.cpp b/Engine/Source/Runtime/Core/Private/Android/AndroidPlatformMisc.cpp
--- a/Engine/Source/Runtime/Core/Private/Android/AndroidPlatformMisc.cpp
+++ b/Engine/Source/Runtime/Core/Private/Android/AndroidPlatformMisc.cpp
@@ -55,6 +55,7 @@
 #include "Misc/OutputDevice.h"
 #include "Logging/LogMacros.h"
 #include "Misc/OutputDeviceError.h"
+#include "Experimental/Async/ConditionVariable.h"
 
 #include "VulkanCommon.h"
 #include "IVulkanDynamicRHI.h"
@@ -2383,8 +2384,19 @@ bool FAndroidMisc::IsExternalMemoryAndroidHardwareBufferExtensionLoaded()
 	return false;
 }
 
+static FCriticalSection ConfigRuleVarsCS;
+static UE::FConditionVariable AvailableConfigRuleVarsCV;
+static std::atomic_bool bConfigRulesReady = false;
+
 const TMap<FString, FString>& FAndroidMisc::GetConfigRuleVars()
 {
+	while(!bConfigRulesReady)
+	{
+		FPlatformMisc::LowLevelOutputDebugStringf(TEXT("thread waiting for configrules to be set"));
+		AvailableConfigRuleVarsCV.Wait(ConfigRuleVarsCS);
+		FPlatformMisc::LowLevelOutputDebugStringf(TEXT("done thread waiting for configrules to be set"));
+	}
+	// ConfigRuleVars are read-only from here on.
 	return ConfigRulesVariables;
 }
 
@@ -2409,6 +2421,8 @@ bool FAndroidMisc::AllowThreadHeartBeat()
 
 JNI_METHOD void Java_com_epicgames_unreal_GameActivity_nativeSetConfigRulesVariables(JNIEnv* jenv, jobject thiz, jobjectArray KeyValuePairs)
 {
+	FScopeLock ConfigRuleVarsLock(&ConfigRuleVarsCS);
+
 	int32 Count = jenv->GetArrayLength(KeyValuePairs);
 	int32 Index = 0;
 	while (Index < Count)
@@ -2418,6 +2432,9 @@ JNI_METHOD void Java_com_epicgames_unreal_GameActivity_nativeSetConfigRulesVaria
 		
 		FAndroidMisc::ConfigRulesVariables.Add(javaKey, javaValue);
 	}
+
+	bConfigRulesReady = true;
+	AvailableConfigRuleVarsCV.NotifyAll();
 }
 
 extern bool AndroidThunkCpp_HasMetaDataKey(const FString& Key);
@@ -2433,6 +2450,12 @@ JNI_METHOD void Java_com_epicgames_unreal_GameActivity_nativeSetAndroidStartupSt
 	}
 }
 
+bool FAndroidMisc::UseNewWindowBehavior()
+{
+	static const FString* EnableNewBackgroundBehavior = FAndroidMisc::GetConfigRulesVariable(TEXT("EnableNewBackgroundBehavior"));
+	return EnableNewBackgroundBehavior && EnableNewBackgroundBehavior->ToLower() == TEXT("true");
+}
+
 #if !UE_BUILD_SHIPPING
 bool FAndroidMisc::IsDebuggerPresent()
 {
diff --git a/Engine/Source/Runtime/Core/Public/Android/AndroidPlatformMisc.h b/Engine/Source/Runtime/Core/Public/Android/AndroidPlatformMisc.h
--- a/Engine/Source/Runtime/Core/Public/Android/AndroidPlatformMisc.h
+++ b/Engine/Source/Runtime/Core/Public/Android/AndroidPlatformMisc.h
@@ -9,7 +9,6 @@
 #include "CoreTypes.h"
 #include "Android/AndroidSystemIncludes.h"
 #include "GenericPlatform/GenericPlatformMisc.h"
-//@todo android: this entire file
 
 template <typename FuncType>
 class TFunction;
@@ -249,6 +248,9 @@ public:
 		__sync_synchronize();
 	}
 
+	// If this true the RHI will not be blocked while the app is in the background. Instead the RHI will detach the window and any further draw operations will be made to an offscreen surface.
+	// false, standard blocking background behavior will be used...
+	static CORE_API bool UseNewWindowBehavior();
 
 #if STATS || ENABLE_STATNAMEDEVENTS
 	/** Override to add Android features when stats are on, otherwise use version on FGenericPlatformMisc. */
diff --git a/Engine/Source/Runtime/Launch/Private/Android/AndroidEventManager.cpp b/Engine/Source/Runtime/Launch/Private/Android/AndroidEventManager.cpp
--- a/Engine/Source/Runtime/Launch/Private/Android/AndroidEventManager.cpp
+++ b/Engine/Source/Runtime/Launch/Private/Android/AndroidEventManager.cpp
@@ -167,7 +167,7 @@ void FAppEventManager::Tick()
 		EmptyQueueHandlerEvent->Trigger();
 	}
 
-	if (!bRunning)
+	if (!bRunning) // && !FAndroidMisc::UseNewWindowBehavior()) // TODO: stalling GT here could be optional
 	{
 		STANDALONE_DEBUG_LOGf(LogAndroidEvents, TEXT("FAppEventManager::Tick EventHandlerEvent Wait "));
 		EventHandlerEvent->Wait();
@@ -291,14 +291,17 @@ void FAppEventManager::SetEmptyQueueHandlerEvent(FEvent* InEventHandlerEvent)
 
 void FAppEventManager::ExecWindowCreated()
 {
-	UE_LOG(LogAndroidEvents, Display, TEXT("ExecWindowCreated"));
+	UE_LOG(LogAndroidEvents, Display, TEXT("ExecWindowCreated (using new behavior %d)"), FAndroidMisc::UseNewWindowBehavior());
 	// When application launched while device is in sleep mode SystemResolution could be set to opposite orientation values
 	// Force to update SystemResolution to current values whenever we create a new window
 	FPlatformRect ScreenRect = FAndroidWindow::GetScreenRect();
 	FSystemResolution::RequestResolutionChange(ScreenRect.Right - ScreenRect.Left, ScreenRect.Bottom - ScreenRect.Top, EWindowMode::Fullscreen);
 
 	// ReInit with the new window handle
-	FAndroidAppEntry::ReInitWindow();
+	// this code has no knowledge of multiple windows.
+	// here we are using whatever the 'main' (only) window is.
+	TOptional<FAndroidWindow::FNativeAccessor> MainWindowAccessor = FAndroidWindowManager::Get().GetMainWindowAsNativeAccessor(false);
+	FAndroidAppEntry::ReInitWindow(MainWindowAccessor);
 	FAndroidApplication::OnWindowSizeChanged();
 }
 
@@ -309,7 +312,12 @@ void FAppEventManager::ExecWindowResized()
 		FlushRenderingCommands();
 	}
 	FAndroidWindow::InvalidateCachedScreenRect();
-	FAndroidAppEntry::ReInitWindow();
+
+	// ReInit with the new window handle
+	// this code has no knowledge of multiple windows.
+	// here we are using whatever the 'main' (only) window is...
+	TOptional<FAndroidWindow::FNativeAccessor> MainWindowAccessor = FAndroidWindowManager::Get().GetMainWindowAsNativeAccessor(false);
+	FAndroidAppEntry::ReInitWindow(MainWindowAccessor);
 	FAndroidApplication::OnWindowSizeChanged();
 }
 
diff --git a/Engine/Source/Runtime/Launch/Private/Android/LaunchAndroid.cpp b/Engine/Source/Runtime/Launch/Private/Android/LaunchAndroid.cpp
--- a/Engine/Source/Runtime/Launch/Private/Android/LaunchAndroid.cpp
+++ b/Engine/Source/Runtime/Launch/Private/Android/LaunchAndroid.cpp
@@ -20,6 +20,7 @@
 #include <sys/system_properties.h>
 #include <dlfcn.h>
 #include "Android/AndroidWindow.h"
+#include "Android/AndroidDynamicRHI.h"
 #include "Android/AndroidApplication.h"
 #include "Android/AndroidPlatformStackWalk.h"
 #include "HAL/PlatformApplicationMisc.h"
@@ -305,13 +306,15 @@ static void ActivateApp_EventThread();
 
 void FPlatformMisc::UnlockAndroidWindow()
 {
-	UE_LOG(LogAndroid, Log, TEXT("Unlocking android HW window during preinit. bAppIsActive_EventThread=%d"), bAppIsActive_EventThread);
-
 	check(IsInGameThread());
 	check(FTaskGraphInterface::IsRunning());
 	
 	bReadyToProcessEvents = true;
-	GAndroidWindowLock.Unlock();
+	if(!FAndroidMisc::UseNewWindowBehavior())
+	{
+		UE_LOG(LogAndroid, Log, TEXT("Unlocking android HW window during preinit. bAppIsActive_EventThread=%d"), bAppIsActive_EventThread);
+		GAndroidWindowLock.Unlock();
+	}
 }
 
 JNI_METHOD void Java_com_epicgames_unreal_GameActivity_nativeResumeMainInit(JNIEnv* jenv, jobject thiz)
@@ -663,6 +666,7 @@ int32 AndroidMain(struct android_app* state)
 	// Have to use a special initialize if using the PersistentStorageManager
 	IPlatformFile::GetPlatformPhysical().Initialize(nullptr, FCommandLine::Get());
 
+	if(!FAndroidMisc::UseNewWindowBehavior())
 	{
 		SCOPED_BOOT_TIMING("Wait for GAndroidWindowLock.Lock()");
 		// wait for a valid window
@@ -673,6 +677,10 @@ int32 AndroidMain(struct android_app* state)
 		UE_LOG(LogAndroid, Log, TEXT("PreInit android HW window lock. bAppIsActive_EventThread=%d"), bAppIsActive_EventThread);
 		GAndroidWindowLock.Lock();
 	}
+	else
+	{
+		UE_LOG(LogAndroid, Log, TEXT("PreInit android using new window behavior. bAppIsActive_EventThread=%d"), bAppIsActive_EventThread);
+	}
 
 	FPlatformMisc::LowLevelOutputDebugString(TEXT("After GAndroidWindowLock in AndroidMain"));
 
@@ -755,6 +763,12 @@ int32 AndroidMain(struct android_app* state)
 		{
 			// use less CPU when paused
 			FPlatformProcess::Sleep(0.10f);
+
+//	TODO:we could continue to tick the engine loop as the new window behavior no longer blocks the RT .
+// 			if (FAndroidMisc::UseNewWindowBehavior())
+// 			{
+// 				GEngineLoop.Tick();
+// 			}
 		}
 
 #if !UE_BUILD_SHIPPING
@@ -934,9 +948,12 @@ static void* AndroidEventThreadWorker( void* param )
 
 	TheChoreographer.SetupChoreographer();
 
-	// window is initially invalid/locked.
-	UE_LOG(LogAndroid, Log, TEXT("AndroidEventThreadWorker, Initial HW window lock."));
-	GAndroidWindowLock.Lock();
+	if(!FAndroidMisc::UseNewWindowBehavior())
+	{
+		// window is initially invalid/locked.
+		UE_LOG(LogAndroid, Log, TEXT("AndroidEventThreadWorker, Initial HW window lock."));
+		GAndroidWindowLock.Lock();
+	}
 
 	DEVELOPER_LOG_COMMANDCB_CASE(AndroidEventThreadWorker_BeforeWhile);
 
@@ -948,7 +965,10 @@ static void* AndroidEventThreadWorker( void* param )
 	}
 	DEVELOPER_LOG_COMMANDCB_CASE(AndroidEventThreadWorker_AfterWhile);
 
-	GAndroidWindowLock.Unlock();
+	if (!FAndroidMisc::UseNewWindowBehavior())
+	{
+		GAndroidWindowLock.Unlock();
+	}
 
 	UE_LOG(LogAndroid, Log, TEXT("AndroidEventThreadWorker->Exiting"));
 	STANDALONE_DEBUG_LOGf(LogAndroid, TEXT("Exiting AndroidEventThreadWorker"));
@@ -1433,7 +1453,6 @@ FAppEventData::FAppEventData(ANativeWindow* WindowIn)
 	//check(WindowWidth >= 0 && WindowHeight >= 0);
 }
 
-// called when the app has window + resume.
 static void ActivateApp_EventThread()
 {
 	DEVELOPER_LOG_COMMANDCB_CASE(ActivateApp_EventThread);
@@ -1442,9 +1461,12 @@ static void ActivateApp_EventThread()
 		return;
 	}
 
-	// Unlock window when we're ready.
-	UE_LOG(LogAndroid, Log, TEXT("event thread, activate app, unlocking HW window"));
-	GAndroidWindowLock.Unlock();
+	if (!FAndroidMisc::UseNewWindowBehavior())
+	{
+		// Unlock window when we're ready.
+		UE_LOG(LogAndroid, Log, TEXT("event thread, activate app, unlocking HW window"));
+		GAndroidWindowLock.Unlock();
+	}
 
 	// wake the GT up.
 	FAppEventManager::GetInstance()->EnqueueAppEvent(APP_EVENT_STATE_APP_ACTIVATED);
@@ -1466,7 +1488,10 @@ static void ActivateApp_EventThread()
 		EventHandlerEvent->Trigger();
 	}
 
-	FThreadHeartBeat::Get().ResumeHeartBeat(true);
+	if (FThreadHeartBeat* ThreadHeartBeat = FThreadHeartBeat::GetNoInit())
+	{
+		ThreadHeartBeat->ResumeHeartBeat(true);
+	}
 
 	FPreLoadScreenManager::EnableRendering(true);
 
@@ -1475,6 +1500,52 @@ static void ActivateApp_EventThread()
 }
 
 extern void BlockRendering();
+
+// called when the app has window with the new behavior model.
+static void OnNewWindow_EventThread(ANativeWindow* NewNativeWindow, ANativeWindow* OldNativeWindow)
+{
+	// Unlock window when we're ready.
+	UE_LOG(LogAndroid, Log, TEXT("event thread, new window"));
+
+	TOptional<FAndroidWindow::FNativeAccessor> WindowAccessor = FAndroidWindowManager::Get().FindFromANativeWindow(true, true, OldNativeWindow);
+	if (WindowAccessor.IsSet())
+	{
+		// Update the window handle, this will be used by the RHI callback.
+		WindowAccessor->SetANativeWindow(NewNativeWindow);
+
+		const auto& OnReinitWindowCallback = UE::FAndroidPlatformDynamicRHI::GetRHIOnReInitWindowCallback();
+		if (OnReinitWindowCallback)
+		{
+			OnReinitWindowCallback(MoveTemp(WindowAccessor));
+		}
+	}
+	else
+	{
+		// no window as yet, save the native handle for the next window create.
+		FAndroidWindowManager::GetPendingWindowAccessor().SetANativeWindow(NewNativeWindow);
+	}
+}
+
+// called whenever the app loses window with the new behavior model.
+void OnLostWindow_EventThread(ANativeWindow* OutgoingWindow)
+{
+	TOptional<FAndroidWindow::FNativeAccessor> WindowAccessor = FAndroidWindowManager::Get().FindFromANativeWindow(true, false, OutgoingWindow);
+	if (WindowAccessor.IsSet())
+	{
+		WindowAccessor->SetANativeWindow(nullptr);
+
+		const auto& OnReleaseWindowCallback = UE::FAndroidPlatformDynamicRHI::GetRHIOnReleaseWindowCallback();
+		if(OnReleaseWindowCallback)
+		{
+			OnReleaseWindowCallback(MoveTemp(WindowAccessor));
+		}
+	}
+	else
+	{
+		FAndroidWindowManager::GetPendingWindowAccessor().SetANativeWindow(nullptr);
+	}
+}
+
 // called whenever the app loses window or pause.
 static void SuspendApp_EventThread()
 {
@@ -1489,8 +1560,11 @@ static void SuspendApp_EventThread()
 	UE::GC::Private::SetSkipDestroy(true);
 	// Lock the window, this prevents event thread from removing the window whilst the RHI initializes.
 
-	UE_LOG(LogAndroid, Log, TEXT("SuspendApp_EventThread -> event thread, suspending app, acquiring HW window lock."));
-	GAndroidWindowLock.Lock();
+	if (!FAndroidMisc::UseNewWindowBehavior())
+	{
+		UE_LOG(LogAndroid, Log, TEXT("SuspendApp_EventThread -> event thread, suspending app, acquiring HW window lock."));
+		GAndroidWindowLock.Lock();
+	}
 	STANDALONE_DEBUG_LOGf(LogAndroid, TEXT("SuspendApp_EventThread -> event thread, suspending app, acquiring HW window lock. bReadyToProcessEvents=%d, bAppIsActive_EventThread=%d"), bReadyToProcessEvents, bAppIsActive_EventThread);
 
 	if (bReadyToProcessEvents == false)
@@ -1536,17 +1610,19 @@ static void SuspendApp_EventThread()
 
 	FThreadHeartBeat::Get().SuspendHeartBeat(true);
 
-	// wait for a period of time before blocking rendering
-	UE_LOG(LogAndroid, Log, TEXT("SuspendApp_EventThread -> , waiting for event manager to process. tid: %d"), FPlatformTLS::GetCurrentThreadId());
-
-	bool bSuccess = EMDoneTrigger->Wait(4000);
+	if (!FAndroidMisc::UseNewWindowBehavior())
+	{
+		// wait for a period of time before blocking rendering
+		UE_LOG(LogAndroid, Log, TEXT("SuspendApp_EventThread -> , waiting for event manager to process. tid: %d"), FPlatformTLS::GetCurrentThreadId());
 
-	float ElapsedTimeInMs_EMDoneTrigger_Wait = FPlatformTime::ToMilliseconds(FPlatformTime::Cycles() - StartCycles);
-	UE_CLOG(!bSuccess, LogAndroid, Log, TEXT("SuspendApp_EventThread -> ERROR: backgrounding callback, not responded in timely manner. EMDoneTrigger->Wait, waited '%f' ms"), (float)ElapsedTimeInMs_EMDoneTrigger_Wait);
-	UE_LOG(LogAndroid, Log, TEXT("SuspendApp_EventThread -> EMDoneTrigger->Wait, waited '%f' ms"), (float)ElapsedTimeInMs_EMDoneTrigger_Wait);
+		bool bSuccess = EMDoneTrigger->Wait(4000);
 
-	BlockRendering();
+		float ElapsedTimeInMs_EMDoneTrigger_Wait = FPlatformTime::ToMilliseconds(FPlatformTime::Cycles() - StartCycles);
+		UE_CLOG(!bSuccess, LogAndroid, Log, TEXT("SuspendApp_EventThread -> ERROR: backgrounding callback, not responded in timely manner. EMDoneTrigger->Wait, waited '%f' ms"), (float)ElapsedTimeInMs_EMDoneTrigger_Wait);
+		UE_LOG(LogAndroid, Log, TEXT("SuspendApp_EventThread -> EMDoneTrigger->Wait, waited '%f' ms"), (float)ElapsedTimeInMs_EMDoneTrigger_Wait);
 
+		BlockRendering();
+	}
 	// Suspend the GT.
 #if !USE_ANDROID_ALTERNATIVE_SUSPEND
 	FAppEventManager::GetInstance()->EnqueueAppEvent(APP_EVENT_STATE_APP_SUSPENDED);
@@ -1560,7 +1636,8 @@ static void OnAppCommandCB(struct android_app* app, int32_t cmd)
 {
 	check(IsInAndroidEventThread());
 
-
+	// Since native app can only have one window we can track of the last handle.
+	static ANativeWindow* LastDestroyedANativeWindow = nullptr;
 	// Set event thread's view of the window dimensions:
 	{
 		ANativeWindow* DimensionWindow = app->pendingWindow ? app->pendingWindow : app->window;
@@ -1595,6 +1672,10 @@ static void OnAppCommandCB(struct android_app* app, int32_t cmd)
 		DEVELOPER_LOG_COMMANDCB_CASE(APP_CMD_INIT_WINDOW);
 		UE_LOG(LogAndroid, Log, TEXT("Case APP_CMD_INIT_WINDOW"));
 
+		if (FAndroidMisc::UseNewWindowBehavior())
+		{
+			OnNewWindow_EventThread(app->window, LastDestroyedANativeWindow);
+		}
 #if USE_ANDROID_STANDALONE
 		{
 			ANativeWindow* win = app->pendingWindow ? app->pendingWindow : app->window;
@@ -1613,9 +1694,12 @@ static void OnAppCommandCB(struct android_app* app, int32_t cmd)
 		bHasWindow = true;
 #endif
 		UE_LOG(LogAndroid, Log, TEXT("Case APP_CMD_INIT_WINDOW: bHasWindow=%d, bHasFocus=%d, bIsResumed=%d"), bHasWindow, bHasFocus, bIsResumed);
-		if (bHasWindow && bHasFocus && bIsResumed)
+		if (!FAndroidMisc::UseNewWindowBehavior())
 		{
-			ActivateApp_EventThread();
+			if (bHasWindow && bHasFocus && bIsResumed)
+			{
+				ActivateApp_EventThread();
+			}
 		}
 		break;
 	case APP_CMD_TERM_WINDOW:
@@ -1629,8 +1713,16 @@ static void OnAppCommandCB(struct android_app* app, int32_t cmd)
 		STANDALONE_DEBUG_LOGf(LogAndroid, TEXT("Case APP_CMD_TERM_WINDOW, tid = %d"), gettid());
 		// clean up the window because it is being hidden/closed
 		UE_LOG(LogAndroid, Log, TEXT("Case APP_CMD_TERM_WINDOW"));
+		LastDestroyedANativeWindow = app->window;
 
-		SuspendApp_EventThread();
+		if (FAndroidMisc::UseNewWindowBehavior())
+		{
+			OnLostWindow_EventThread(app->window);
+		}
+		else
+		{
+			SuspendApp_EventThread();
+		}
 
 		FAppEventManager::GetInstance()->HandleWindowClosed_EventThread();
 		bHasWindow = false;
@@ -1658,7 +1750,7 @@ static void OnAppCommandCB(struct android_app* app, int32_t cmd)
 		UE_LOG(LogAndroid, Log, TEXT("Case APP_CMD_GAINED_FOCUS"));
 
 		// still check for a rare case needing activation
-		if (bHasWindow && bHasFocus && bIsResumed)
+		if ((bHasWindow || FAndroidMisc::UseNewWindowBehavior()) && bHasFocus && bIsResumed)
 		{
 			ActivateApp_EventThread();
 		}
@@ -1736,7 +1828,7 @@ static void OnAppCommandCB(struct android_app* app, int32_t cmd)
 		bHasFocus = true;
 		DEVELOPER_LOG_COMMANDCB_CASE(APP_CMD_RESUME);
 
-		if (bHasWindow && bHasFocus && bIsResumed)
+		if ((bHasWindow || FAndroidMisc::UseNewWindowBehavior()) && bHasFocus && bIsResumed)
 		{
 			ActivateApp_EventThread();
 		}
diff --git a/Engine/Source/Runtime/OpenGLDrv/Private/Android/AndroidEGL.cpp b/Engine/Source/Runtime/OpenGLDrv/Private/Android/AndroidEGL.cpp
--- a/Engine/Source/Runtime/OpenGLDrv/Private/Android/AndroidEGL.cpp
+++ b/Engine/Source/Runtime/OpenGLDrv/Private/Android/AndroidEGL.cpp
@@ -107,9 +107,10 @@ struct AndroidESPImpl
 	ANativeWindow* Window             = nullptr;
 	GLuint         ResolveFrameBuffer = 0;
 	GLuint         DummyFrameBuffer   = 0;
-	FPlatformRect  CachedWindowRect {};
+	FPlatformRect  CachedWindowRect     {};
 	bool           Initalized         = false;
 	bool           bIsDebug           = false;
+	bool           bIsWndSurface      = false; // True when the surface is attached to a HW window.
 
 	AndroidESPImpl() = default;
 };
@@ -241,8 +242,11 @@ void AndroidEGL::CreateEGLRenderSurface(ANativeWindow* InWindow, bool bCreateWnd
 		return;
 	}
 
+	PImplData->bIsWndSurface = bCreateWndSurface;
+
 	if (bCreateWndSurface)
 	{
+		check(InWindow);
 		//need ANativeWindow
 		PImplData->eglSurface = eglCreateWindowSurface(PImplData->eglDisplay, PImplData->eglConfigParam,InWindow, NULL);
 
@@ -462,7 +466,7 @@ void AndroidEGL::InitEGL(APIVariant API)
 // call out to JNI to see if the application was packaged for Oculus Mobile
 extern bool AndroidThunkCpp_IsOculusMobileApplication();
 
-void AndroidEGL::SetRenderContextWindowSurface()
+void AndroidEGL::SetRenderContextWindowSurface(const TOptional<FAndroidWindow::FNativeAccessor>& WindowContainer)
 {
 	VERIFY_GL_SCOPE();
 
@@ -471,23 +475,26 @@ void AndroidEGL::SetRenderContextWindowSurface()
 	UnBindRender();
 
 	SetCurrentContext(EGL_NO_CONTEXT, EGL_NO_SURFACE);
-	{
-		bool bCreateSurface = !AndroidThunkCpp_IsOculusMobileApplication();
+	bool bCreateSurface = !AndroidThunkCpp_IsOculusMobileApplication();
 
+	if (!FAndroidMisc::UseNewWindowBehavior())
+	{
 		// SetRenderContextWindowSurface is called only when the window lock is successful.
 		PImplData->Window = (ANativeWindow*)FAndroidWindow::GetHardwareWindow_EventThread();
 		check(PImplData->Window);
-		InitRenderSurface(false, bCreateSurface);
 	}
+
+	InitRenderSurface(false, bCreateSurface, WindowContainer);
 	SetCurrentRenderingContext();
 
 	FPlatformMisc::LowLevelOutputDebugStringf(TEXT("AndroidEGL::SetRenderContextWindowSurface  DONE! tid: %d"),  FPlatformTLS::GetCurrentThreadId());
 }
 
-void AndroidEGL::ResizeRenderContextSurface()
+void AndroidEGL::ResizeRenderContextSurface(const TOptional<FAndroidWindow::FNativeAccessor>& WindowContainer)
 {
 	VERIFY_GL_SCOPE();
 
+	check(!FAndroidMisc::UseNewWindowBehavior() || !WindowContainer.IsSet() || WindowContainer.GetValue().GetANativeWindow() == PImplData->Window);
 	// Resize render originates from the gamethread, we cant use Window_Event here.
 	if (PImplData->Window)
 	{
@@ -508,7 +515,7 @@ void AndroidEGL::ResizeRenderContextSurface()
 			SetCurrentContext(EGL_NO_CONTEXT, EGL_NO_SURFACE);
 			{
 				bool bCreateSurface = !AndroidThunkCpp_IsOculusMobileApplication();
-				InitRenderSurface(false, bCreateSurface);
+				InitRenderSurface(false, bCreateSurface, WindowContainer);
 			}
 			SetCurrentRenderingContext();
 		}
@@ -549,11 +556,40 @@ void AndroidEGL::InitBackBuffer()
 
 extern void AndroidThunkCpp_SetDesiredViewSize(int32 Width, int32 Height);
 
-void AndroidEGL::InitRenderSurface(bool bUseSmallSurface, bool bCreateWndSurface)
+void AndroidEGL::InitRenderSurface(bool bUseSmallSurface, bool bCreateWndSurface, const TOptional<FAndroidWindow::FNativeAccessor>& WindowContainer)
 {
 	FPlatformMisc::LowLevelOutputDebugStringf(TEXT("AndroidEGL::InitRenderSurface %d, %d"), int(bUseSmallSurface), int(bCreateWndSurface));
 
-	check(PImplData->Window);
+	if (FAndroidMisc::UseNewWindowBehavior())
+	{
+		bUseSmallSurface = true;
+		bCreateWndSurface = false;
+		PImplData->Window = nullptr;
+
+		if (WindowContainer.IsSet())
+		{
+			const FAndroidWindow::FNativeAccessor& Accessor = WindowContainer.GetValue();
+			if(Accessor.GetANativeWindow())
+			{
+				bUseSmallSurface = false;
+				bCreateWndSurface = true;
+				PImplData->Window = (ANativeWindow*) Accessor.GetANativeWindow();
+				FPlatformMisc::LowLevelOutputDebugStringf(TEXT("AndroidEGL::InitRenderSurface window %p was supplied %d, %d"), PImplData->Window, int(bUseSmallSurface), int(bCreateWndSurface));
+			}
+			else
+			{
+				FPlatformMisc::LowLevelOutputDebugStringf(TEXT("AndroidEGL::InitRenderSurface null native window was supplied %d, %d"), int(bUseSmallSurface), int(bCreateWndSurface));
+			}
+		}
+		else
+		{
+			FPlatformMisc::LowLevelOutputDebugStringf(TEXT("AndroidEGL::InitRenderSurface No window was supplied %d, %d"), int(bUseSmallSurface), int(bCreateWndSurface));
+		}
+	}
+	else
+	{
+		check(PImplData->Window);
+	}
 
 	int32 Width = 8, Height = 8;
 	if (!bUseSmallSurface)
@@ -583,30 +619,18 @@ void AndroidEGL::InitRenderSurface(bool bUseSmallSurface, bool bCreateWndSurface
 		AndroidThunkCpp_SetDesiredViewSize(Width, Height);
 	}
 
-	FIntVector2 OriginalWindowSize(ANativeWindow_getWidth(PImplData->Window), ANativeWindow_getHeight(PImplData->Window));
+	if(PImplData->Window)
+	{
+		FIntVector2 OriginalWindowSize(ANativeWindow_getWidth(PImplData->Window), ANativeWindow_getHeight(PImplData->Window));
 
-	FPlatformMisc::LowLevelOutputDebugStringf(TEXT("AndroidEGL::InitRenderSurface, setting wnd: %p, width: %d->%d, height %d->%d "), PImplData->Window, OriginalWindowSize.X, Width, OriginalWindowSize.Y, Height);
-	ANativeWindow_setBuffersGeometry(PImplData->Window, Width, Height, PImplData->NativeVisualID);
+		FPlatformMisc::LowLevelOutputDebugStringf(TEXT("AndroidEGL::InitRenderSurface, setting wnd: %p, width: %d->%d, height %d->%d "), PImplData->Window, OriginalWindowSize.X, Width, OriginalWindowSize.Y, Height);
+		ANativeWindow_setBuffersGeometry(PImplData->Window, Width, Height, PImplData->NativeVisualID);
+	}
 	CreateEGLRenderSurface(PImplData->Window, bCreateWndSurface);
 
 	PImplData->RenderingContext.eglSurface = PImplData->eglSurface;
 }
 
-void AndroidEGL::ReInit()
-{
-	FPlatformMisc::LowLevelOutputDebugString(TEXT("AndroidEGL::ReInit()"));
-
-	FPlatformRect WindowRect = FAndroidWindow::GetScreenRect();
-	FPlatformMisc::LowLevelOutputDebugStringf(TEXT("AndroidEGL FAndroidAppEntry::ReInitWindow updating window size = %d, %d, cached size: %d, %d tid: %d"), WindowRect.Right, WindowRect.Bottom, PImplData->CachedWindowRect.Right, PImplData->CachedWindowRect.Bottom, FPlatformTLS::GetCurrentThreadId());
-	PImplData->CachedWindowRect = WindowRect;
-
-	SetCurrentContext(EGL_NO_CONTEXT, EGL_NO_SURFACE);
-	bool bCreateSurface = !AndroidThunkCpp_IsOculusMobileApplication();
-	InitRenderSurface(false, bCreateSurface);
-	
-	SetCurrentContext(PImplData->RenderingContext.eglContext, PImplData->RenderingContext.eglSurface);
-}
-
 // EGL_CONTEXT_OPENGL_ROBUST_ACCESS_EXT is enabled if configrules asks for it or the command line specifies it.
 // If -OpenGLRobustContext=[0/1] is specified on the command line it takes precedence.
 void AndroidEGL::Init(APIVariant API, uint32 MajorVersion, uint32 MinorVersion)
@@ -701,8 +725,11 @@ void AndroidEGL::Init(APIVariant API, uint32 MajorVersion, uint32 MinorVersion)
 		}
 	}
 
-	// Getting the hardware window is valid during preinit as we have GAndroidWindowLock held.
-	PImplData->Window = (ANativeWindow*)FAndroidWindow::GetHardwareWindow_EventThread();
+	if (!FAndroidMisc::UseNewWindowBehavior())
+	{
+		// Getting the hardware window is valid during preinit as we have GAndroidWindowLock held.
+		PImplData->Window = (ANativeWindow*)FAndroidWindow::GetHardwareWindow_EventThread();
+	}
 	PImplData->Initalized   = true;
 }
 
@@ -769,6 +796,11 @@ EGLConfig AndroidEGL::GetConfig() const
 	return PImplData->eglConfigParam;
 }
 
+bool AndroidEGL::IsUsingWindowedSurface() const
+{
+	return PImplData->bIsWndSurface;
+}
+
 void AndroidEGL::GetSwapIntervalRange(EGLint& OutMinSwapInterval, EGLint& OutMaxSwapInterval) const
 {
 	eglGetConfigAttrib(PImplData->eglDisplay, PImplData->eglConfigParam, EGL_MIN_SWAP_INTERVAL, &OutMinSwapInterval);
@@ -866,7 +898,7 @@ void AndroidEGL::UnBindRender()
 	DestroyRenderSurface();
 }
 
-void FAndroidAppEntry::ReInitWindow(void* NewNativeWindowHandle)
+void FAndroidAppEntry::ReInitWindow(const TOptional<FAndroidWindow::FNativeAccessor>& WindowContainer)
 {
 	check(IsInGameThread());
 
@@ -887,11 +919,11 @@ void FAndroidAppEntry::ReInitWindow(void* NewNativeWindowHandle)
 	if (!FAndroidMisc::ShouldUseVulkan())
 	{
 		// the window size could have been adjusted by the GT by now, if so it must be updated.
-		AndroidEGL::GetInstance()->RefreshWindowSize();
+		AndroidEGL::GetInstance()->RefreshWindowSize(WindowContainer);
 	}
 }
 
-void AndroidEGL::RefreshWindowSize()
+void AndroidEGL::RefreshWindowSize(const TOptional<FAndroidWindow::FNativeAccessor>& WindowContainer)
 {
 	check(IsInGameThread());
 	check(!FAndroidMisc::ShouldUseVulkan());
@@ -900,13 +932,15 @@ void AndroidEGL::RefreshWindowSize()
 	PImplData->CachedWindowRect = WindowRect;
 
 	ENQUEUE_RENDER_COMMAND(EGLResizeRenderContextSurface)(
-		[](FRHICommandListImmediate& RHICmdList)
+		[&WindowContainer](FRHICommandListImmediate& RHICmdList) mutable
 	{
-		RHICmdList.EnqueueLambda([](FRHICommandListImmediate&)
+		RHICmdList.EnqueueLambda([&WindowContainer](FRHICommandListImmediate&) mutable
 		{
-			AndroidEGL::GetInstance()->ResizeRenderContextSurface();
+			AndroidEGL::GetInstance()->ResizeRenderContextSurface(WindowContainer);
 		});
 	});
+
+	FlushRenderingCommands();
 }
 
 void FAndroidAppEntry::OnPauseEvent()
@@ -1054,7 +1088,7 @@ void BlockOnLostWindowRenderCommand(TSharedPtr<FEvent, ESPMode::ThreadSafe> RTBl
 			RTBlockedTrigger->Trigger();
 			GAndroidWindowLock.Lock();
 			UE_LOG(LogAndroid, Log, TEXT("RendererBlock acquired window lock"));
-			AndroidEGL::GetInstance()->SetRenderContextWindowSurface();
+			AndroidEGL::GetInstance()->SetRenderContextWindowSurface(TOptional<FAndroidWindow::FNativeAccessor>()); // FNativeAccessor is ignored with previous window behavior.
 			UE_LOG(LogAndroid, Log, TEXT("RendererBlock updating window"));
 			GAndroidWindowLock.Unlock();
 		});
diff --git a/Engine/Source/Runtime/OpenGLDrv/Private/Android/AndroidEGL.h b/Engine/Source/Runtime/OpenGLDrv/Private/Android/AndroidEGL.h
--- a/Engine/Source/Runtime/OpenGLDrv/Private/Android/AndroidEGL.h
+++ b/Engine/Source/Runtime/OpenGLDrv/Private/Android/AndroidEGL.h
@@ -14,6 +14,7 @@
 #include <EGL/egl.h>
 #include <EGL/eglext.h>
 #include <GLES3/gl31.h>
+#include "Android/AndroidWindow.h"
 
 struct AndroidESPImpl;
 struct ANativeWindow;
@@ -74,13 +75,15 @@ public:
 	void DestroyBackBuffer();
 
 	void Init(APIVariant API, uint32 MajorVersion, uint32 MinorVersion);
-	void ReInit();
 
 	void UnBindRender();
 
 	void Terminate();
 
-	void InitRenderSurface(bool bUseSmallSurface, bool bCreateWndSurface);
+	// Initialize the surface of the context.
+	// when using the new window locking behavior, if WindowContainer is not set an off screen surface will be used.
+	// TODO: remove bCreateWndSurface (and possibly bUseSmallSurface), these can be inferred from WindowContainer.
+	void InitRenderSurface(bool bUseSmallSurface, bool bCreateWndSurface, const TOptional<FAndroidWindow::FNativeAccessor>& WindowContainer);
 	void UpdateBuffersTransform();
 	bool IsOfflineSurfaceRequired();
 
@@ -110,10 +113,13 @@ public:
 	bool GetSupportsNoErrorContext();
 
 	// recreate the EGL surface for the current hardware window.
-	void SetRenderContextWindowSurface();
+	void SetRenderContextWindowSurface(const TOptional<FAndroidWindow::FNativeAccessor>& WindowHandle);
 
 	// Called from game thread when a window is reinited.
-	void RefreshWindowSize();
+	void RefreshWindowSize(const TOptional<FAndroidWindow::FNativeAccessor>& WindowContainer);
+
+	// true if the current surface is associated with a window.
+	bool IsUsingWindowedSurface() const;
 
 protected:
 	AndroidEGL();
@@ -137,7 +143,7 @@ private:
 	void LogConfigInfo(EGLConfig  EGLConfigInfo);
 
 	// Actual Update to the egl surface to match the GT's requested size.
-	void ResizeRenderContextSurface();
+	void ResizeRenderContextSurface(const TOptional<FAndroidWindow::FNativeAccessor>& WindowContainer);
 
 	bool bSupportsKHRCreateContext      = false;
 	bool bSupportsKHRSurfacelessContext = false;
diff --git a/Engine/Source/Runtime/OpenGLDrv/Private/Android/AndroidOpenGL.cpp b/Engine/Source/Runtime/OpenGLDrv/Private/Android/AndroidOpenGL.cpp
--- a/Engine/Source/Runtime/OpenGLDrv/Private/Android/AndroidOpenGL.cpp
+++ b/Engine/Source/Runtime/OpenGLDrv/Private/Android/AndroidOpenGL.cpp
@@ -76,6 +76,7 @@ extern bool AndroidThunkCpp_IsOculusMobileApplication();
 #define GL_DEBUG_TOOL_EXT	0x6789
 static bool bRunningUnderRenderDoc = false;
 
+extern bool IsInAndroidEventThread();
 
 #if UE_BUILD_SHIPPING
 #define CHECK_JNI_EXCEPTIONS(env)  env->ExceptionClear();
@@ -173,7 +174,8 @@ void FPlatformOpenGLDevice::Init()
 
 	FPlatformMisc::LowLevelOutputDebugString(TEXT("FPlatformOpenGLDevice:Init"));
 	bool bCreateSurface = !AndroidThunkCpp_IsOculusMobileApplication();
-	AndroidEGL::GetInstance()->InitRenderSurface(false, bCreateSurface);
+	// With the new window system we may not have a window yet, the surface will be initially off screen.
+	AndroidEGL::GetInstance()->InitRenderSurface(false, bCreateSurface, TOptional<FAndroidWindow::FNativeAccessor>());
 
 	LoadEXT();
 	PlatformRenderingContextSetup(this);
@@ -216,6 +218,65 @@ void FPlatformOpenGLDevice::Init()
 		GRHISupportsGPUTimestampBubblesRemoval = true;
 	}
 #endif
+
+	// register the new window behavior life cycle callbacks.
+	if (FPlatformMisc::UseNewWindowBehavior())
+	{
+		UE::FAndroidPlatformDynamicRHI::SetRHIOnReleaseWindowCallback(
+			[](const TOptional<FAndroidWindow::FNativeAccessor>& WindowContainer) 
+			{
+				check(IsInAndroidEventThread());
+				UE_LOG(LogRHI, Log, TEXT("OnReleaseWindowCallback event thread"));
+				FGraphEventRef OnComplete = FGraphEvent::CreateGraphEvent();
+
+				ENQUEUE_RENDER_COMMAND(OnAndroidLostWindow)([&WindowContainer, OnComplete](FRHICommandListImmediate& RHICmdList) mutable
+					{
+						UE_LOG(LogRHI, Log, TEXT("OnReleaseWindowCallback: RT"));
+						RHICmdList.EnqueueLambda([&WindowContainer, OnComplete](FRHICommandListImmediate& ExecutingCmdList) mutable
+							{
+								UE_LOG(LogRHI, Log, TEXT("GLES OnReleaseWindowCallback: RHI, set egl surface to offscreen (window %p)"), WindowContainer->GetANativeWindow());
+								AndroidEGL::GetInstance()->UnBindRender();
+								AndroidEGL::GetInstance()->InitRenderSurface(true, false, WindowContainer);
+								AndroidEGL::GetInstance()->SetCurrentRenderingContext();
+								UE_LOG(LogRHI, Log, TEXT("GLES OnReleaseWindowCallback: RHI, swap surface done."));
+							});
+						RHICmdList.RHIThreadFence(true);
+						RHICmdList.ImmediateFlush(EImmediateFlushType::FlushRHIThread);
+						OnComplete->DispatchSubsequents();
+						UE_LOG(LogRHI, Log, TEXT("OnReleaseWindowCallback: RT done"));
+					});
+				UE_LOG(LogRHI, Log, TEXT("OnLostWindow: waiting for RT"));
+				OnComplete->Wait();
+				UE_LOG(LogRHI, Log, TEXT("OnLostWindow: done"));
+			});
+
+		UE::FAndroidPlatformDynamicRHI::SetRHIOnReInitWindowCallback(
+			[](const TOptional<FAndroidWindow::FNativeAccessor>& WindowContainer)
+			{
+				UE_LOG(LogRHI, Log, TEXT("OnReInitWindowCallback event thread"));
+				FGraphEventRef OnComplete = FGraphEvent::CreateGraphEvent();
+
+				ENQUEUE_RENDER_COMMAND(OnAndroidFoundWindow)([&WindowContainer, OnComplete](FRHICommandListImmediate& RHICmdList) mutable
+					{
+						UE_LOG(LogRHI, Log, TEXT("OnReInitWindowCallback: RT"));
+						RHICmdList.EnqueueLambda([&WindowContainer](FRHICommandListImmediate& ExecutingCmdList) mutable
+							{
+								// TODO: we should be searching for the affected viewport here.
+								UE_LOG(LogRHI, Log, TEXT("GLES OnReInitWindowCallback: RHI, set egl surface to %p"), WindowContainer->GetANativeWindow());
+								AndroidEGL::GetInstance()->SetRenderContextWindowSurface(WindowContainer);
+								// gles can swap out the display surface immediately.
+								UE_LOG(LogRHI, Log, TEXT("SWP: GLES OnReInitWindowCallback: RHI, done"));
+							});
+						RHICmdList.RHIThreadFence(true);
+						RHICmdList.ImmediateFlush(EImmediateFlushType::FlushRHIThread);
+						OnComplete->DispatchSubsequents();
+						UE_LOG(LogRHI, Log, TEXT("OnReInitWindowCallback: RT done"));
+					});
+				UE_LOG(LogRHI, Log, TEXT("OnReInitWindowCallback: waiting for RT"));
+				OnComplete->Wait();
+				UE_LOG(LogRHI, Log, TEXT("OnReInitWindowCallback: done"));
+			});
+	}
 }
 
 FPlatformOpenGLDevice* PlatformCreateOpenGLDevice()
@@ -447,8 +508,18 @@ void FPlatformOpenGLDevice::LoadEXT()
 
 FPlatformOpenGLContext* PlatformCreateOpenGLContext(FPlatformOpenGLDevice* Device, void* InWindowHandle)
 {
-	//Assumes Device is already initialized and context already created.
-	return AndroidEGL::GetInstance()->GetRenderingContext();
+	AndroidEGL* AnEGL = AndroidEGL::GetInstance();
+	if (FAndroidMisc::UseNewWindowBehavior())
+	{
+		check(InWindowHandle);
+		FAndroidWindow* AndroidWindow = (FAndroidWindow *)InWindowHandle;
+
+		FAndroidWindow::FNativeAccessor Accessor = AndroidWindow->GetANativeAccessor(false);
+
+		AnEGL->InitRenderSurface(false, true, MoveTemp(Accessor));
+	}
+	// else - with the original window method we block on startup until the window + device are initialized and context fully created.
+	return AnEGL->GetRenderingContext();
 }
 
 void PlatformDestroyOpenGLContext(FPlatformOpenGLDevice* Device, FPlatformOpenGLContext* Context)
diff --git a/Engine/Source/Runtime/OpenGLDrv/Private/Android/AndroidOpenGLFramePacer.cpp b/Engine/Source/Runtime/OpenGLDrv/Private/Android/AndroidOpenGLFramePacer.cpp
--- a/Engine/Source/Runtime/OpenGLDrv/Private/Android/AndroidOpenGLFramePacer.cpp
+++ b/Engine/Source/Runtime/OpenGLDrv/Private/Android/AndroidOpenGLFramePacer.cpp
@@ -239,6 +239,8 @@ bool FAndroidOpenGLFramePacer::SupportsFramePace(int32 QueryFramePace)
 	return SupportsFramePaceInternal(QueryFramePace, TempRefreshRate, TempSyncInterval);
 }
 
+extern bool AndroidThunkCpp_IsOculusMobileApplication();
+
 bool FAndroidOpenGLFramePacer::SwapBuffers(bool bLockToVsync)
 {
 	SCOPED_NAMED_EVENT(STAT_OpenGLSwapBuffersTime, FColor::Red)
@@ -253,17 +255,25 @@ bool FAndroidOpenGLFramePacer::SwapBuffers(bool bLockToVsync)
 
 	VERIFY_EGL_SCOPE();
 
-	EGLDisplay eglDisplay = AndroidEGL::GetInstance()->GetDisplay();
-	EGLSurface eglSurface = AndroidEGL::GetInstance()->GetSurface();
+	AndroidEGL* AnEGL = AndroidEGL::GetInstance();
+	EGLDisplay eglDisplay = AnEGL->GetDisplay();
+	EGLSurface eglSurface = AnEGL->GetSurface();
 	int32 SyncInterval = FAndroidPlatformRHIFramePacer::GetLegacySyncInterval();
 
+	if (AnEGL->IsUsingWindowedSurface() == false && !AndroidThunkCpp_IsOculusMobileApplication())
+	{
+		// offscreen surface, no swap to perform, limit update rate.
+		FPlatformProcess::Sleep(0.1);
+		return true;
+	}
+
 	bool bPrintMethod = false;
 
 #if USE_ANDROID_OPENGL_SWAPPY
 	int32 CurrentFramePace = 0;
 	if (FAndroidPlatformRHIFramePacer::CVarUseSwappyForFramePacing.GetValueOnRenderThread() != 0 && (CurrentFramePace = FAndroidPlatformRHIFramePacer::GetFramePace()) != 0 && ensure(bSwappyInit))
 	{
-		ANativeWindow* CurrentNativeWindow = AndroidEGL::GetInstance()->GetNativeWindow();
+		ANativeWindow* CurrentNativeWindow = AnEGL->GetNativeWindow();
 		if (CurrentNativeWindow != CachedNativeWindow)
 		{
 			CachedNativeWindow = CurrentNativeWindow;
@@ -374,7 +384,7 @@ bool FAndroidOpenGLFramePacer::SwapBuffers(bool bLockToVsync)
 
 			// make sure requested interval is in supported range
 			EGLint MinSwapInterval, MaxSwapInterval;
-			AndroidEGL::GetInstance()->GetSwapIntervalRange(MinSwapInterval, MaxSwapInterval);
+			AnEGL->GetSwapIntervalRange(MinSwapInterval, MaxSwapInterval);
 
 			int64 SyncIntervalNanos = (30 + 1000000000l * int64(SyncInterval)) / 60;
 
diff --git a/Engine/Source/Runtime/OpenGLDrv/Private/Android/AndroidOpenGLPrivate.cpp b/Engine/Source/Runtime/OpenGLDrv/Private/Android/AndroidOpenGLPrivate.cpp
--- a/Engine/Source/Runtime/OpenGLDrv/Private/Android/AndroidOpenGLPrivate.cpp
+++ b/Engine/Source/Runtime/OpenGLDrv/Private/Android/AndroidOpenGLPrivate.cpp
@@ -26,7 +26,8 @@ FAndroidGPUInfo::FAndroidGPUInfo()
 	// Do not create a window surface if the app is for Oculus Mobile (use small buffer)
 	bool bCreateSurface = !AndroidThunkCpp_IsOculusMobileApplication();
 	FPlatformMisc::LowLevelOutputDebugString(TEXT("FAndroidGPUInfo"));
-	EGL->InitRenderSurface(false, bCreateSurface);
+
+	EGL->InitRenderSurface(false, bCreateSurface, TOptional<FAndroidWindow::FNativeAccessor>());
 	EGL->SetCurrentRenderingContext();
 
 	// get extensions
diff --git a/Engine/Source/Runtime/OpenGLDrv/Private/OpenGLDrv.cpp b/Engine/Source/Runtime/OpenGLDrv/Private/OpenGLDrv.cpp
--- a/Engine/Source/Runtime/OpenGLDrv/Private/OpenGLDrv.cpp
+++ b/Engine/Source/Runtime/OpenGLDrv/Private/OpenGLDrv.cpp
@@ -101,7 +101,7 @@ bool FOpenGLDynamicRHI::RHIEGLSupportsNoErrorContext() const
 void FOpenGLDynamicRHI::RHIInitEGLInstanceGLES2()
 {
 	AndroidEGL::GetInstance()->Init(AndroidEGL::AV_OpenGLES, 2, 0);
-	AndroidEGL::GetInstance()->InitRenderSurface(false, false);
+	AndroidEGL::GetInstance()->InitRenderSurface(false, false, TOptional<FAndroidWindow::FNativeAccessor>());
 }
 
 void FOpenGLDynamicRHI::RHIInitEGLBackBuffer()
diff --git a/Engine/Source/Runtime/RHI/Private/Android/AndroidDynamicRHI.cpp b/Engine/Source/Runtime/RHI/Private/Android/AndroidDynamicRHI.cpp
--- a/Engine/Source/Runtime/RHI/Private/Android/AndroidDynamicRHI.cpp
+++ b/Engine/Source/Runtime/RHI/Private/Android/AndroidDynamicRHI.cpp
@@ -92,6 +92,31 @@ namespace UE
 		{
 			return GExternalCompilerFailureThreshold;
 		}
+
+		static FRHIReInitWindowCallbackType OnRHIReleaseWindowCallback;
+
+		FRHIReInitWindowCallbackType& GetRHIOnReInitWindowCallback()
+		{
+			return OnRHIReleaseWindowCallback;
+		}
+
+		void SetRHIOnReInitWindowCallback(FRHIReInitWindowCallbackType&& InRHIOnReInitWindowCallback)
+		{
+			OnRHIReleaseWindowCallback = MoveTemp(InRHIOnReInitWindowCallback);
+		}
+
+		static FRHIReleaseWindowCallbackType OnRHIReleaseWindowCallbackType;
+
+		FRHIReleaseWindowCallbackType& GetRHIOnReleaseWindowCallback()
+		{
+			return OnRHIReleaseWindowCallbackType;
+		}
+
+		void SetRHIOnReleaseWindowCallback(FRHIReleaseWindowCallbackType&& InRHIOnReleaseWindowCallback)
+		{
+			OnRHIReleaseWindowCallbackType = MoveTemp(InRHIOnReleaseWindowCallback);
+		}
+
 	} // namespace FAndroidPlatformDynamicRHI
 } // namespace UE
 
diff --git a/Engine/Source/Runtime/RHI/Public/Android/AndroidDynamicRHI.h b/Engine/Source/Runtime/RHI/Public/Android/AndroidDynamicRHI.h
--- a/Engine/Source/Runtime/RHI/Public/Android/AndroidDynamicRHI.h
+++ b/Engine/Source/Runtime/RHI/Public/Android/AndroidDynamicRHI.h
@@ -2,11 +2,21 @@
 #pragma once
 
 #include "RHI.h"
+#include "Android/AndroidWindow.h"
 
 namespace UE
 {
 	namespace FAndroidPlatformDynamicRHI
 	{
+		typedef TUniqueFunction<void(TOptional<FAndroidWindow::FNativeAccessor> WindowContainer)> FRHIReInitWindowCallbackType;
+		FRHIReInitWindowCallbackType& GetRHIOnReInitWindowCallback();
+		void SetRHIOnReInitWindowCallback(FRHIReInitWindowCallbackType&& InRHIOnReInitWindowCallback);
+
+		typedef TUniqueFunction<void(TOptional<FAndroidWindow::FNativeAccessor> WindowContainer)> FRHIReleaseWindowCallbackType;
+		FRHIReleaseWindowCallbackType& GetRHIOnReleaseWindowCallback();
+		void SetRHIOnReleaseWindowCallback(FRHIReleaseWindowCallbackType&& InOnReleaseWindowCallback);
+
+
 		// Helper class to convert FGraphicsPipelineStateInitializer::EPSOPrecacheCompileType to the uint64 required by android's external PSO service.
 		class FPSOServicePriInfo
 		{
diff --git a/Engine/Source/Runtime/VulkanRHI/Private/Android/VulkanAndroidPlatform.cpp b/Engine/Source/Runtime/VulkanRHI/Private/Android/VulkanAndroidPlatform.cpp
--- a/Engine/Source/Runtime/VulkanRHI/Private/Android/VulkanAndroidPlatform.cpp
+++ b/Engine/Source/Runtime/VulkanRHI/Private/Android/VulkanAndroidPlatform.cpp
@@ -12,6 +12,7 @@
 #include "Android/AndroidPlatformMisc.h"
 #include "Android/AndroidJavaEnv.h"
 #include "Android/AndroidJNI.h"
+#include "Android/AndroidPlatformMisc.h"
 #include "Misc/ConfigCacheIni.h"
 #include "GenericPlatform/GenericPlatformCrashContext.h"
 #include "../VulkanExtensions.h"
@@ -20,6 +21,7 @@
 #include "ProfilingDebugging/ScopedTimers.h"
 #include "Android/AndroidDynamicRHI.h"
 #include "PSOMetrics.h"
+#include "../VulkanContext.h"
 
 #if USE_ANDROID_VULKAN_SWAPPY
 #undef VK_NO_PROTOTYPES
@@ -416,15 +418,14 @@ bool FVulkanAndroidPlatform::SupportsBCTextureFormats()
 	return (CVarVulkanSupportsBCTextureFormats.GetValueOnAnyThread() == 1);
 }
 
-void FVulkanAndroidPlatform::CreateSurface(void* WindowHandle, VkInstance Instance, VkSurfaceKHR* OutSurface)
+void FVulkanAndroidPlatform::CreateSurface(FVulkanPlatformWindowContext& WindowContext, VkInstance Instance, VkSurfaceKHR* OutSurface)
 {
-	// don't use cached window handle coming from VulkanViewport, as it could be gone by now
-	WindowHandle = GetHardwareWindowHandle();
+	ANativeWindow* Handle = FAndroidMisc::UseNewWindowBehavior() ? WindowContext.GetANativeWindow() : (ANativeWindow*)GetHardwareWindowHandle();
+	check(Handle);
 
 	VkAndroidSurfaceCreateInfoKHR SurfaceCreateInfo;
 	ZeroVulkanStruct(SurfaceCreateInfo, VK_STRUCTURE_TYPE_ANDROID_SURFACE_CREATE_INFO_KHR);
-	SurfaceCreateInfo.window = (ANativeWindow*)WindowHandle;
-
+	SurfaceCreateInfo.window = Handle;
 	VERIFYVULKANRESULT(VulkanDynamicAPI::vkCreateAndroidSurfaceKHR(Instance, &SurfaceCreateInfo, VULKAN_CPU_ALLOCATOR, OutSurface));
 }
 
@@ -634,18 +635,206 @@ bool FVulkanAndroidPlatform::SupportsDynamicResolution()
 #endif
 }
 
+extern bool IsInAndroidEventThread();
+
+namespace AndroidVulkan
+{
+	static void OnLostWindow(TOptional<FAndroidWindow::FNativeAccessor> WindowContainer)
+	{
+		check(IsInAndroidEventThread());
+		check(FAndroidMisc::UseNewWindowBehavior());
+		UE_LOG(LogVulkanRHI, Log, TEXT("OnLostWindow: Event thread"));
+		FVulkanPlatformWindowContext WindowContext(MoveTemp(WindowContainer));
+		FGraphEventRef OnComplete = FGraphEvent::CreateGraphEvent();
+
+		ENQUEUE_RENDER_COMMAND(OnAndroidLostWindow)([OnComplete, &WindowContext](FRHICommandListImmediate& RHICmdList) mutable
+		{
+				UE_LOG(LogVulkanRHI, Log, TEXT("OnLostWindow: RT"));
+				RHICmdList.EnqueueLambda([&WindowContext](FRHICommandListImmediate& ExecutingCmdList) mutable
+					{
+						check(IsInRHIThread() || IsRunningRHIInSeparateThread());
+						UE_LOG(LogVulkanRHI, Log, TEXT("OnLostWindow: RHI"));
+
+						// cannot just throw out this command buffer (needs to be submitted or other checks fail)
+						FVulkanCommandListContext& Context = FVulkanCommandListContext::Get(ExecutingCmdList);
+						Context.FlushCommands(EVulkanFlushFlags::WaitForSubmission);
+
+						bool bFound = false;
+						TArray<FVulkanViewport*> Viewports = GVulkanRHI->GetViewports();
+						for (FVulkanViewport* Viewport : Viewports)
+						{
+							if (Viewport->GetWindowHandle() == WindowContext.GetWindowHandle())
+							{
+								UE_LOG(LogVulkanRHI, Log, TEXT("OnLostWindow: Destroying viewport swapchain (window %p -> null)"), WindowContext.GetANativeWindow() );
+								Viewport->DestroySwapchain(nullptr);
+								UE_LOG(LogVulkanRHI, Log, TEXT("OnLostWindow: Destroying viewport swapchain done."));
+								bFound = true;
+								break;
+							}
+						}
+						check(bFound);
+						UE_LOG(LogVulkanRHI, Log, TEXT("OnLostWindow: RHI done"));
+					});
+
+				// Swapchain creation pushes some commands - flush the command buffers now to begin with a fresh state
+				RHICmdList.RHIThreadFence(true);
+				RHICmdList.SubmitAndBlockUntilGPUIdle();
+
+				UE_LOG(LogVulkanRHI, Log, TEXT("OnLostWindow:RT done"));
+				OnComplete->DispatchSubsequents();
+			});
+
+		UE_LOG(LogVulkanRHI, Log, TEXT("OnLostWindow: wait RT"));
+		OnComplete->Wait();
+		UE_LOG(LogVulkanRHI, Log, TEXT("OnLostWindow: done"));
+	}
+
+	static void OnFoundWindow(TOptional<FAndroidWindow::FNativeAccessor> WindowContainer)
+	{
+		check(IsInAndroidEventThread());
+		check(FAndroidMisc::UseNewWindowBehavior());
+
+		FVulkanPlatformWindowContext WindowContext(MoveTemp(WindowContainer));
+
+		UE_LOG(LogVulkanRHI, Log, TEXT("OnFoundWindow:"));
+		FGraphEventRef OnComplete = FGraphEvent::CreateGraphEvent();
+		ENQUEUE_RENDER_COMMAND(OnAndroidFoundWindow)([OnComplete, &WindowContext](FRHICommandListImmediate& RHICmdList) mutable
+			{
+				UE_LOG(LogVulkanRHI, Log, TEXT("OnFoundWindow: RT"));
+				RHICmdList.EnqueueLambda([&WindowContext](FRHICommandListImmediate& ExecutingCmdList) mutable
+					{
+						check(IsInRHIThread() || IsRunningRHIInSeparateThread());
+						UE_LOG(LogVulkanRHI, Log, TEXT("OnFoundWindow: RHI"));
+
+						FVulkanDevice* Device = FVulkanDynamicRHI::Get().GetDevice();
+
+						// cannot just throw out this command buffer (needs to be submitted or other checks fail)
+						FVulkanCommandListContext& Context = FVulkanCommandListContext::Get(ExecutingCmdList);
+						Context.FlushCommands(EVulkanFlushFlags::WaitForSubmission);
+
+						bool bFound = false;
+						TArray<FVulkanViewport*> Viewports = GVulkanRHI->GetViewports();
+						for (FVulkanViewport* Viewport : Viewports)
+						{
+							if (Viewport->GetWindowHandle() == WindowContext.GetWindowHandle())
+							{
+								UE_LOG(LogVulkanRHI, Log, TEXT("OnFoundWindow: RHI RecreateSwapchain (window %p)"), WindowContext.GetANativeWindow());
+								Viewport->RecreateSwapchain(Context, WindowContext);
+								UE_LOG(LogVulkanRHI, Log, TEXT("OnFoundWindow: RHI RecreateSwapchain done (window %p)"), WindowContext.GetANativeWindow());
+								bFound = true;
+								break;
+							}
+						}
+						check(bFound);
+						UE_LOG(LogVulkanRHI, Log, TEXT("OnFoundWindow: RHI done"));
+					});
+				// Swapchain creation pushes some commands - flush the command buffers now to begin with a fresh state
+				RHICmdList.RHIThreadFence(true);
+				RHICmdList.SubmitAndBlockUntilGPUIdle();
+				UE_LOG(LogVulkanRHI, Log, TEXT("OnFoundWindow: RT done"));
+				OnComplete->DispatchSubsequents();
+			});
+		UE_LOG(LogVulkanRHI, Log, TEXT("OnFoundWindow: wait RT"));
+		OnComplete->Wait();
+		UE_LOG(LogVulkanRHI, Log, TEXT("OnFoundWindow: done"));
+	}
+}
+
+
+// Old window behavior window recreate function.
+// Moved to platform specific code from VulkanRHI.cpp.
+void FVulkanAndroidPlatform::RecreateSwapChain(void* NewNativeWindow)
+{
+	check(!FAndroidMisc::UseNewWindowBehavior());
+
+	if (NewNativeWindow)
+	{
+		if (IsInGameThread())
+		{
+			FlushRenderingCommands();
+		}
+
+		TArray<FVulkanViewport*> Viewports = FVulkanDynamicRHI::Get().GetViewports();
+		ENQUEUE_RENDER_COMMAND(VulkanRecreateSwapChain)(
+			[Viewports = MoveTemp(Viewports), NewNativeWindow](FRHICommandListImmediate& RHICmdList) mutable
+		{
+			UE_LOG(LogVulkanRHI, Log, TEXT("Recreate swapchain ... "));
+			RHICmdList.EnqueueLambda(TEXT("RecreateSwapchain"), [Viewports = MoveTemp(Viewports), NewNativeWindow](FRHICommandListImmediate& ExecutingCmdList)
+			{
+				for (FVulkanViewport* Viewport : Viewports)
+				{
+					FVulkanPlatformWindowContext WindowContext(NewNativeWindow);
+					Viewport->RecreateSwapchain(FVulkanCommandListContext::Get(ExecutingCmdList), WindowContext);
+				}
+			});
+
+			RHICmdList.SubmitAndBlockUntilGPUIdle();
+		});
+
+		if (IsInGameThread())
+		{
+			FlushRenderingCommands();
+		}
+	}
+}
+
+// Old window behavior window destroy function.
+// Moved to platform specific code from VulkanRHI.cpp.
+void FVulkanAndroidPlatform::DestroySwapChain()
+{
+	check(!FAndroidMisc::UseNewWindowBehavior());
+
+	if (IsInGameThread())
+	{
+		FlushRenderingCommands();
+	}
+
+	TArray<FVulkanViewport*> Viewports = FVulkanDynamicRHI::Get().GetViewports();
+	ENQUEUE_RENDER_COMMAND(VulkanDestroySwapChain)(
+		[Viewports = MoveTemp(Viewports)](FRHICommandListImmediate& RHICmdList) mutable
+		{
+			UE_LOG(LogVulkanRHI, Log, TEXT("Destroy swapchain ... "));
+
+			RHICmdList.EnqueueLambda(TEXT("FVulkanDynamicRHI::DestroySwapChain"), [Viewports = MoveTemp(Viewports)](FRHICommandListBase& ExecutingCmdList)
+				{
+					for (FVulkanViewport* Viewport : Viewports)
+					{
+						Viewport->DestroySwapchain(nullptr);
+					}
+				});
+			RHICmdList.ImmediateFlush(EImmediateFlushType::FlushRHIThread);
+		});
+
+	if (IsInGameThread())
+	{
+		FlushRenderingCommands();
+	}
+}
+
+
 void FVulkanAndroidPlatform::OverridePlatformHandlers(bool bInit)
 {
 	if (bInit)
 	{
-		FPlatformMisc::SetOnReInitWindowCallback(FVulkanDynamicRHI::RecreateSwapChain);
-		FPlatformMisc::SetOnReleaseWindowCallback(FVulkanDynamicRHI::DestroySwapChain);
+		if (FAndroidMisc::UseNewWindowBehavior())
+		{
+			UE::FAndroidPlatformDynamicRHI::SetRHIOnReleaseWindowCallback(AndroidVulkan::OnLostWindow);
+			UE::FAndroidPlatformDynamicRHI::SetRHIOnReInitWindowCallback(AndroidVulkan::OnFoundWindow);
+		}
+		else
+		{
+			FPlatformMisc::SetOnReInitWindowCallback(FVulkanAndroidPlatform::RecreateSwapChain);
+			FPlatformMisc::SetOnReleaseWindowCallback(FVulkanAndroidPlatform::DestroySwapChain);
+		}
 		FPlatformMisc::SetOnPauseCallback(FVulkanDynamicRHI::SavePipelineCache);
 	}
 	else
 	{
 		FPlatformMisc::SetOnReInitWindowCallback(nullptr);
 		FPlatformMisc::SetOnReleaseWindowCallback(nullptr);
+		UE::FAndroidPlatformDynamicRHI::SetRHIOnReleaseWindowCallback(nullptr);
+		UE::FAndroidPlatformDynamicRHI::SetRHIOnReInitWindowCallback(nullptr);
+
 		FPlatformMisc::SetOnPauseCallback(nullptr);
 	}
 }
@@ -727,7 +916,7 @@ VkResult FVulkanAndroidPlatform::Present(VkQueue Queue, VkPresentInfoKHR& Presen
 	}
 }
 
-VkResult FVulkanAndroidPlatform::CreateSwapchainKHR(void* WindowHandle, VkPhysicalDevice PhysicalDevice, VkDevice Device, const VkSwapchainCreateInfoKHR* CreateInfo, const VkAllocationCallbacks* Allocator, VkSwapchainKHR* Swapchain)
+VkResult FVulkanAndroidPlatform::CreateSwapchainKHR(FVulkanPlatformWindowContext& WindowContext, VkPhysicalDevice PhysicalDevice, VkDevice Device, const VkSwapchainCreateInfoKHR* CreateInfo, const VkAllocationCallbacks* Allocator, VkSwapchainKHR* Swapchain)
 {
 	VkResult Result = VulkanRHI::vkCreateSwapchainKHR(Device, CreateInfo, Allocator, Swapchain);
 
@@ -739,13 +928,12 @@ VkResult FVulkanAndroidPlatform::CreateSwapchainKHR(void* WindowHandle, VkPhysic
 			JNIEnv* Env = FAndroidApplication::GetJavaEnv();
 			if (ensure(Env))
 			{
-				// don't use cached window handle coming from VulkanViewport, as it could be gone by now
-				WindowHandle = GetHardwareWindowHandle();
+				ANativeWindow* WindowHandle = FAndroidMisc::UseNewWindowBehavior() ? WindowContext.GetANativeWindow() : (ANativeWindow*)GetHardwareWindowHandle();
 				check(WindowHandle);
 								
 				uint64_t RefreshDuration; // in nanoseconds
 				SwappyVk_initAndGetRefreshCycleDuration(Env, FJavaWrapper::GameActivityThis, PhysicalDevice, Device, *Swapchain, &RefreshDuration);
-				SwappyVk_setWindow(Device, *Swapchain, (ANativeWindow*)WindowHandle);	
+				SwappyVk_setWindow(Device, *Swapchain, WindowHandle);	
 				SwappyVk_setAutoSwapInterval(false);
 				AndroidVulkan::SetSwappyPostWaitCallback();
 				UE_LOG(LogVulkanRHI, Log, TEXT("SwappyVk_initAndGetRefreshCycleDuration: %ull"), RefreshDuration);
diff --git a/Engine/Source/Runtime/VulkanRHI/Private/Android/VulkanAndroidPlatform.h b/Engine/Source/Runtime/VulkanRHI/Private/Android/VulkanAndroidPlatform.h
--- a/Engine/Source/Runtime/VulkanRHI/Private/Android/VulkanAndroidPlatform.h
+++ b/Engine/Source/Runtime/VulkanRHI/Private/Android/VulkanAndroidPlatform.h
@@ -3,6 +3,7 @@
 #pragma once
 
 #include "RHI.h"
+#include "Android/AndroidWindow.h"
 
 #define VK_USE_PLATFORM_ANDROID_KHR					1
 
@@ -39,7 +40,64 @@
 	EnumMacro(PFN_vkGetPastPresentationTimingGOOGLE, vkGetPastPresentationTimingGOOGLE)
 
 // and now, include the GenericPlatform class
-#include "../VulkanGenericPlatform.h"
+#include "VulkanGenericPlatform.h"
+
+extern bool IsInAndroidEventThread();
+
+class FVulkanAndroidPlatformWindowContext
+{
+public:
+	FVulkanAndroidPlatformWindowContext(void* InWindowHandle)
+	{
+		AndroidWindow = (FAndroidWindow*)InWindowHandle;
+		// in this case FVulkanAndroidPlatformWindowContext owns the FNativeAccessor.
+		if (InWindowHandle)
+		{
+			if(FAndroidMisc::UseNewWindowBehavior())
+			{
+				// the context should be locked at the beginning of the update process, this overload is used for GT initiated events.
+				check(IsInGameThread());
+				WindowContainer = AndroidWindow->GetANativeAccessor(false);
+				NativeWindow = WindowContainer->GetANativeWindow();
+			}
+		}
+	}
+	
+	FVulkanAndroidPlatformWindowContext(TOptional<FAndroidWindow::FNativeAccessor> WindowContainerIn)
+	{
+		// in this case the FNativeAccessor has come from the event thread
+		if(WindowContainerIn.IsSet())
+		{
+			check(FAndroidMisc::UseNewWindowBehavior());
+			// the context should be locked at the beginning of the update process, this overload is used for ET initiated events.
+			check(IsInAndroidEventThread());
+			WindowContainer = MoveTemp(WindowContainerIn);
+			NativeWindow = WindowContainer->GetANativeWindow();
+			AndroidWindow = WindowContainer.GetValue().Get();
+		}
+	}
+
+	// we dont have a locked window to create swapchains during present, android will use the 
+	// SetRHIOnReleaseWindowCallback / SetRHIOnReInitWindowCallback to create as required.
+	static bool CanCreateSwapchainOnDemand() { return false; }
+
+	ANativeWindow* GetANativeWindow() const
+	{
+		return NativeWindow;
+	}
+
+	void* GetWindowHandle() const 
+	{ 
+		return AndroidWindow;
+	}
+
+private:
+	ANativeWindow* NativeWindow = nullptr;
+	FAndroidWindow* AndroidWindow = nullptr;
+	TOptional<FAndroidWindow::FNativeAccessor> WindowContainer;
+};
+
+using FVulkanPlatformWindowContext = FVulkanAndroidPlatformWindowContext;
 
 class FVulkanAndroidPlatform : public FVulkanGenericPlatform
 {
@@ -57,7 +115,7 @@ public:
 	static void GetDeviceLayers(TArray<const ANSICHAR*>& OutLayers);
 	static void NotifyFoundDeviceLayersAndExtensions(VkPhysicalDevice PhysicalDevice, const TArray<const ANSICHAR*>& Layers, const TArray<const ANSICHAR*>& Extensions);
 
-	static void CreateSurface(void* WindowHandle, VkInstance Instance, VkSurfaceKHR* OutSurface);
+	static void CreateSurface(FVulkanPlatformWindowContext& WindowContext, VkInstance Instance, VkSurfaceKHR* OutSurface);
 
 	static void* GetHardwareWindowHandle();
 
@@ -143,7 +201,7 @@ public:
 
 	static VkResult Present(VkQueue Queue, VkPresentInfoKHR& PresentInfo);
 
-	static VkResult CreateSwapchainKHR(void* WindowHandle, VkPhysicalDevice PhysicalDevice, VkDevice Device, const VkSwapchainCreateInfoKHR* CreateInfo, const VkAllocationCallbacks* Allocator, VkSwapchainKHR* Swapchain);
+	static VkResult CreateSwapchainKHR(FVulkanPlatformWindowContext& WindowContext, VkPhysicalDevice PhysicalDevice, VkDevice Device, const VkSwapchainCreateInfoKHR* CreateInfo, const VkAllocationCallbacks* Allocator, VkSwapchainKHR* Swapchain);
 
 	static void DestroySwapchainKHR(VkDevice Device, VkSwapchainKHR Swapchain, const VkAllocationCallbacks* Allocator);
 
@@ -157,6 +215,11 @@ public:
 	// Do not attempt to immediately recreate swapchain
 	static bool RecreateSwapchainOnFail() { return false; }
 
+	// original window behavior swapchain functions.
+	// ignored with the new window method.
+	static void RecreateSwapChain(void* NewNativeWindow);
+	static void DestroySwapChain();
+
 	static VkFormat GetPlatform5551FormatWithFallback(VkFormat& OutFallbackFormat0, VkFormat& OutFallbackFormat1) 
 	{ 
 		OutFallbackFormat0 = VK_FORMAT_A1R5G5B5_UNORM_PACK16; 
diff --git a/Engine/Source/Runtime/VulkanRHI/Private/Linux/VulkanLinuxPlatform.cpp b/Engine/Source/Runtime/VulkanRHI/Private/Linux/VulkanLinuxPlatform.cpp
--- a/Engine/Source/Runtime/VulkanRHI/Private/Linux/VulkanLinuxPlatform.cpp
+++ b/Engine/Source/Runtime/VulkanRHI/Private/Linux/VulkanLinuxPlatform.cpp
@@ -243,11 +243,11 @@ void FVulkanLinuxPlatform::GetDeviceExtensions(FVulkanDevice* Device, FVulkanDev
 	OutExtensions.Add(MakeUnique<FVulkanDeviceExtension>(Device, VK_KHR_EXTERNAL_MEMORY_FD_EXTENSION_NAME, VULKAN_EXTENSION_ENABLED, VULKAN_EXTENSION_NOT_PROMOTED, nullptr, FVulkanExtensionBase::ManuallyActivate));
 }
 
-void FVulkanLinuxPlatform::CreateSurface(void* WindowHandle, VkInstance Instance, VkSurfaceKHR* OutSurface)
+void FVulkanLinuxPlatform::CreateSurface(FVulkanPlatformWindowContext& WindowContext, VkInstance Instance, VkSurfaceKHR* OutSurface)
 {
 	EnsureSDLIsInited();
 
-	if (SDL_Vulkan_CreateSurface((SDL_Window*)WindowHandle, Instance, OutSurface) == SDL_FALSE)
+	if (SDL_Vulkan_CreateSurface((SDL_Window*)WindowContext.GetWindowHandle(), Instance, OutSurface) == SDL_FALSE)
 	{
 		UE_LOG(LogInit, Error, TEXT("Error initializing SDL Vulkan Surface: %hs"), SDL_GetError());
 		check(0);
diff --git a/Engine/Source/Runtime/VulkanRHI/Private/Linux/VulkanLinuxPlatform.h b/Engine/Source/Runtime/VulkanRHI/Private/Linux/VulkanLinuxPlatform.h
--- a/Engine/Source/Runtime/VulkanRHI/Private/Linux/VulkanLinuxPlatform.h
+++ b/Engine/Source/Runtime/VulkanRHI/Private/Linux/VulkanLinuxPlatform.h
@@ -31,6 +31,8 @@
 // and now, include the GenericPlatform class
 #include "../VulkanGenericPlatform.h"
 
+typedef FVulkanGenericPlatformWindowContext FVulkanPlatformWindowContext;
+
 class FVulkanLinuxPlatform : public FVulkanGenericPlatform
 {
 public:
@@ -45,7 +47,7 @@ public:
 	static void GetDeviceExtensions(FVulkanDevice* Device, FVulkanDeviceExtensionArray& OutExtensions);
 	static void GetDeviceLayers(TArray<const ANSICHAR*>& OutLayers) {}
 
-	static void CreateSurface(void* WindowHandle, VkInstance Instance, VkSurfaceKHR* OutSurface);
+	static void CreateSurface(FVulkanPlatformWindowContext& WindowContext, VkInstance Instance, VkSurfaceKHR* OutSurface);
 
 	static void WriteCrashMarker(const FOptionalVulkanDeviceExtensions& OptionalExtensions, FVulkanCommandBuffer* CmdBuffer, VkBuffer DestBuffer, const TArrayView<uint32>& Entries, bool bAdding);
 
diff --git a/Engine/Source/Runtime/VulkanRHI/Private/VulkanDynamicRHI.h b/Engine/Source/Runtime/VulkanRHI/Private/VulkanDynamicRHI.h
--- a/Engine/Source/Runtime/VulkanRHI/Private/VulkanDynamicRHI.h
+++ b/Engine/Source/Runtime/VulkanRHI/Private/VulkanDynamicRHI.h
@@ -286,9 +286,6 @@ public:
 		return InstanceLayers;
 	}
 
-	static void DestroySwapChain();
-	static void RecreateSwapChain(void* NewNativeWindow);
-
 	VkInstance GetInstance() const
 	{
 		return Instance;
diff --git a/Engine/Source/Runtime/VulkanRHI/Private/VulkanGenericPlatform.cpp b/Engine/Source/Runtime/VulkanRHI/Private/VulkanGenericPlatform.cpp
--- a/Engine/Source/Runtime/VulkanRHI/Private/VulkanGenericPlatform.cpp
+++ b/Engine/Source/Runtime/VulkanRHI/Private/VulkanGenericPlatform.cpp
@@ -118,7 +118,7 @@ VkResult FVulkanGenericPlatform::Present(VkQueue Queue, VkPresentInfoKHR& Presen
 	return VulkanRHI::vkQueuePresentKHR(Queue, &PresentInfo);
 }
 
-VkResult FVulkanGenericPlatform::CreateSwapchainKHR(void* WindowHandle, VkPhysicalDevice PhysicalDevice, VkDevice Device, const VkSwapchainCreateInfoKHR* CreateInfo, const VkAllocationCallbacks* Allocator, VkSwapchainKHR* Swapchain)
+VkResult FVulkanGenericPlatform::CreateSwapchainKHR(FVulkanGenericPlatformWindowContext& WindowContext, VkPhysicalDevice PhysicalDevice, VkDevice Device, const VkSwapchainCreateInfoKHR* CreateInfo, const VkAllocationCallbacks* Allocator, VkSwapchainKHR* Swapchain)
 {
 	return VulkanRHI::vkCreateSwapchainKHR(Device, CreateInfo, Allocator, Swapchain);
 }
diff --git a/Engine/Source/Runtime/VulkanRHI/Private/VulkanGenericPlatform.h b/Engine/Source/Runtime/VulkanRHI/Private/VulkanGenericPlatform.h
--- a/Engine/Source/Runtime/VulkanRHI/Private/VulkanGenericPlatform.h
+++ b/Engine/Source/Runtime/VulkanRHI/Private/VulkanGenericPlatform.h
@@ -23,6 +23,18 @@ using FVulkanDeviceExtensionArray = TArray<TUniquePtr<class FVulkanDeviceExtensi
 using FVulkanInstanceExtensionArray = TArray<TUniquePtr<class FVulkanInstanceExtension>>;
 
 // the platform interface, and empty implementations for platforms that don't need em
+
+class FVulkanGenericPlatformWindowContext
+{
+public:
+	FVulkanGenericPlatformWindowContext(void* WindowHandleIn) : WindowHandle(WindowHandleIn) { }
+	void* GetWindowHandle() const { return WindowHandle; }
+	// whether the platform can attempt to create a swapchains during present
+	static bool CanCreateSwapchainOnDemand() { return true; }
+protected:
+	void* WindowHandle = nullptr;
+};
+
 class FVulkanGenericPlatform
 {
 public:
@@ -103,7 +115,7 @@ public:
 	static VkResult Present(VkQueue Queue, VkPresentInfoKHR& PresentInfo);
 
 	// Allow platforms to track swapchain creation
-	static VkResult CreateSwapchainKHR(void* WindowHandle, VkPhysicalDevice PhysicalDevice, VkDevice Device, const VkSwapchainCreateInfoKHR* CreateInfo, const VkAllocationCallbacks* Allocator, VkSwapchainKHR* Swapchain);
+	static VkResult CreateSwapchainKHR(FVulkanGenericPlatformWindowContext& WindowContext, VkPhysicalDevice PhysicalDevice, VkDevice Device, const VkSwapchainCreateInfoKHR* CreateInfo, const VkAllocationCallbacks* Allocator, VkSwapchainKHR* Swapchain);
 	
 	// Allow platforms to track swapchain destruction
 	static void DestroySwapchainKHR(VkDevice Device, VkSwapchainKHR Swapchain, const VkAllocationCallbacks* Allocator);
diff --git a/Engine/Source/Runtime/VulkanRHI/Private/VulkanRHI.cpp b/Engine/Source/Runtime/VulkanRHI/Private/VulkanRHI.cpp
--- a/Engine/Source/Runtime/VulkanRHI/Private/VulkanRHI.cpp
+++ b/Engine/Source/Runtime/VulkanRHI/Private/VulkanRHI.cpp
@@ -2093,63 +2093,6 @@ void FVulkanDynamicRHI::TrimLRU()
 }
 #endif
 
-void FVulkanDynamicRHI::DestroySwapChain()
-{
-	if (IsInGameThread())
-	{
-		FlushRenderingCommands();
-	}
-
-	TArray<FVulkanViewport*> Viewports = GVulkanRHI->Viewports;
-	ENQUEUE_RENDER_COMMAND(VulkanDestroySwapChain)(
-		[Viewports](FRHICommandListImmediate& RHICmdList)
-	{
-		UE_LOG(LogVulkanRHI, Log, TEXT("Destroy swapchain ... "));
-
-		RHICmdList.ImmediateFlush(EImmediateFlushType::FlushRHIThread);
-
-		for (FVulkanViewport* Viewport : Viewports)
-		{
-			Viewport->DestroySwapchain(nullptr);
-		}
-	});
-
-	if (IsInGameThread())
-	{
-		FlushRenderingCommands();
-	}
-}
-
-void FVulkanDynamicRHI::RecreateSwapChain(void* NewNativeWindow)
-{
-	if (NewNativeWindow)
-	{
-		if (IsInGameThread())
-		{
-			FlushRenderingCommands();
-		}
-
-		TArray<FVulkanViewport*> Viewports = GVulkanRHI->Viewports;
-		ENQUEUE_RENDER_COMMAND(VulkanRecreateSwapChain)(
-			[Viewports, NewNativeWindow](FRHICommandListImmediate& RHICmdList)
-		{
-			UE_LOG(LogVulkanRHI, Log, TEXT("Recreate swapchain ... "));
-			
-			for (FVulkanViewport* Viewport : Viewports)
-			{
-				Viewport->RecreateSwapchainFromRT(RHICmdList, Viewport->PixelFormat, NewNativeWindow);
-			}
-
-			RHICmdList.SubmitAndBlockUntilGPUIdle();
-		});
-
-		if (IsInGameThread())
-		{
-			FlushRenderingCommands();
-		}
-	}
-}
-
 void FVulkanDynamicRHI::VulkanSetImageLayout(VkCommandBuffer CmdBuffer, VkImage Image, VkImageLayout OldLayout, VkImageLayout NewLayout, const VkImageSubresourceRange& SubresourceRange)
 {
 	FVulkanPipelineBarrier Barrier;
diff --git a/Engine/Source/Runtime/VulkanRHI/Private/VulkanSwapChain.cpp b/Engine/Source/Runtime/VulkanRHI/Private/VulkanSwapChain.cpp
--- a/Engine/Source/Runtime/VulkanRHI/Private/VulkanSwapChain.cpp
+++ b/Engine/Source/Runtime/VulkanRHI/Private/VulkanSwapChain.cpp
@@ -148,12 +148,12 @@ bool IsVulkanPreTransformEnabled(const FStaticShaderPlatform Platform)
 			&& CVarVulkanPreTransform.GetValueOnAnyThread() > 0;
 }
 
-FVulkanSwapChain::FVulkanSwapChain(VkInstance InInstance, FVulkanDevice& InDevice, void* InWindowHandle, EPixelFormat& InOutPixelFormat, uint32 Width, uint32 Height, bool bIsFullScreen,
-	uint32* InOutDesiredNumBackBuffers, TArray<VkImage>& OutImages, int8 InLockToVsync, FVulkanSwapChainRecreateInfo* RecreateInfo)
+FVulkanSwapChain::FVulkanSwapChain(VkInstance InInstance, FVulkanDevice& InDevice, EPixelFormat& InOutPixelFormat, uint32 Width, uint32 Height, bool bIsFullScreen,
+	uint32* InOutDesiredNumBackBuffers, TArray<VkImage>& OutImages, int8 InLockToVsync, FVulkanPlatformWindowContext& WindowContext, FVulkanSwapChainRecreateInfo* RecreateInfo)
 	: SwapChain(VK_NULL_HANDLE)
 	, Device(InDevice)
 	, Surface(VK_NULL_HANDLE)
-	, WindowHandle(InWindowHandle)
+	, WindowHandle(WindowContext.GetWindowHandle())
 	, CurrentImageIndex(-1)
 	, SemaphoreIndex(0)
 	, NumPresentCalls(0)
@@ -172,7 +172,7 @@ FVulkanSwapChain::FVulkanSwapChain(VkInstance InInstance, FVulkanDevice& InDevic
 	else
 	{
 		// let the platform create the surface
-		FVulkanPlatform::CreateSurface(WindowHandle, Instance, &Surface);
+		FVulkanPlatform::CreateSurface(WindowContext, Instance, &Surface);
 	}
 
 	uint32 NumFormats;
@@ -545,7 +545,7 @@ FVulkanSwapChain::FVulkanSwapChain(VkInstance InInstance, FVulkanDevice& InDevic
 		Swap(SwapChainInfo.imageExtent.width, SwapChainInfo.imageExtent.height);
 	}
 
-	VkResult Result = FVulkanPlatform::CreateSwapchainKHR(WindowHandle, Device.GetPhysicalHandle(), Device.GetInstanceHandle(), &SwapChainInfo, VULKAN_CPU_ALLOCATOR, &SwapChain);
+	VkResult Result = FVulkanPlatform::CreateSwapchainKHR(WindowContext, Device.GetPhysicalHandle(), Device.GetInstanceHandle(), &SwapChainInfo, VULKAN_CPU_ALLOCATOR, &SwapChain);
 
 	// Swap the aspect ratio back to make it transparent to high level
 	if (bSwapChainNeedsTransform)
@@ -559,7 +559,7 @@ FVulkanSwapChain::FVulkanSwapChain(VkInstance InInstance, FVulkanDevice& InDevic
 		// Unlink fullscreen
 		UE_LOG(LogVulkanRHI, Warning, TEXT("Create swapchain failed with Initialization error; removing FullScreen extension..."));
 		SwapChainInfo.pNext = FullScreenInfo.pNext;
-		Result = FVulkanPlatform::CreateSwapchainKHR(WindowHandle, Device.GetPhysicalHandle(), Device.GetInstanceHandle(), &SwapChainInfo, VULKAN_CPU_ALLOCATOR, &SwapChain);
+		Result = FVulkanPlatform::CreateSwapchainKHR(WindowContext, Device.GetPhysicalHandle(), Device.GetInstanceHandle(), &SwapChainInfo, VULKAN_CPU_ALLOCATOR, &SwapChain);
 	}
 #endif
 	VERIFYVULKANRESULT_EXPANDED(Result);
diff --git a/Engine/Source/Runtime/VulkanRHI/Private/VulkanSwapChain.h b/Engine/Source/Runtime/VulkanRHI/Private/VulkanSwapChain.h
--- a/Engine/Source/Runtime/VulkanRHI/Private/VulkanSwapChain.h
+++ b/Engine/Source/Runtime/VulkanRHI/Private/VulkanSwapChain.h
@@ -33,8 +33,8 @@ struct FVulkanSwapChainRecreateInfo
 class FVulkanSwapChain
 {
 public:
-	FVulkanSwapChain(VkInstance InInstance, FVulkanDevice& InDevice, void* WindowHandle, EPixelFormat& InOutPixelFormat, uint32 Width, uint32 Height, bool bIsFullscreen,
-		uint32* InOutDesiredNumBackBuffers, TArray<VkImage>& OutImages, int8 bLockToVsync, FVulkanSwapChainRecreateInfo* RecreateInfo);
+	FVulkanSwapChain(VkInstance InInstance, FVulkanDevice& InDevice, EPixelFormat& InOutPixelFormat, uint32 Width, uint32 Height, bool bIsFullscreen,
+		uint32* InOutDesiredNumBackBuffers, TArray<VkImage>& OutImages, int8 bLockToVsync, FVulkanPlatformWindowContext& PlatformContext, FVulkanSwapChainRecreateInfo* RecreateInfo);
 
 	void Destroy(FVulkanSwapChainRecreateInfo* RecreateInfo);
 
diff --git a/Engine/Source/Runtime/VulkanRHI/Private/VulkanTexture.cpp b/Engine/Source/Runtime/VulkanRHI/Private/VulkanTexture.cpp
--- a/Engine/Source/Runtime/VulkanRHI/Private/VulkanTexture.cpp
+++ b/Engine/Source/Runtime/VulkanRHI/Private/VulkanTexture.cpp
@@ -1490,7 +1490,7 @@ FVulkanTexture::FVulkanTexture(FVulkanDevice& InDevice, const FRHITextureCreateD
 		}
 	}
 
-	const bool bIsTransientResource = (InTransientHeapAllocation != nullptr);
+	const bool bIsTransientResource = (InTransientHeapAllocation != nullptr) && InTransientHeapAllocation->IsValid();
 	if (bIsTransientResource)
 	{
 		check(!bMemoryless);
@@ -1620,8 +1620,11 @@ FVulkanTexture::FVulkanTexture(FVulkanDevice& InDevice, const FRHITextureCreateD
 			const bool bDoInitialClear = bRenderTarget;
 			const VkImageLayout InitialLayout = VK_IMAGE_LAYOUT_UNDEFINED;  // use undefinied to avoid transitioning the texture when aliasing
 
-			FRHICommandList& RHICmdList = FRHICommandListExecutor::GetImmediateCommandList();
-			SetInitialImageState(RHICmdList, InitialLayout, bDoInitialClear, InCreateDesc.ClearValue, false);
+			if (!EnumHasAnyFlags(InCreateDesc.Flags, TexCreate_Presentable))
+			{
+				FRHICommandList& RHICmdList = FRHICommandListExecutor::GetImmediateCommandList();
+				SetInitialImageState(RHICmdList, InitialLayout, bDoInitialClear, InCreateDesc.ClearValue, false);
+			}
 		}
 	}
 
diff --git a/Engine/Source/Runtime/VulkanRHI/Private/VulkanViewport.cpp b/Engine/Source/Runtime/VulkanRHI/Private/VulkanViewport.cpp
--- a/Engine/Source/Runtime/VulkanRHI/Private/VulkanViewport.cpp
+++ b/Engine/Source/Runtime/VulkanRHI/Private/VulkanViewport.cpp
@@ -132,11 +132,15 @@ FVulkanViewport::FVulkanViewport(FVulkanDevice* InDevice, void* InWindowHandle,
 
 	bRenderOffscreen = FParse::Param(FCommandLine::Get(), TEXT("RenderOffScreen"));
 
-	FlushRenderingCommands();
+	FVulkanPlatformWindowContext WindowContext(InWindowHandle);
+
 	ENQUEUE_RENDER_COMMAND(CreateSwapchain)(
-		[VulkanViewport=this](FRHICommandListImmediate& RHICmdList)
+		[&WindowContext, VulkanViewport = this](FRHICommandListImmediate& RHICmdList)
 		{
-			VulkanViewport->CreateSwapchain(RHICmdList, nullptr);
+			RHICmdList.EnqueueLambda([&WindowContext, VulkanViewport](FRHICommandListBase& ExecutingCmdList)
+				{
+					VulkanViewport->CreateSwapchain(FVulkanCommandListContext::Get(ExecutingCmdList), nullptr, WindowContext);
+				});
 			RHICmdList.SubmitAndBlockUntilGPUIdle();
 		});
 	FlushRenderingCommands();
@@ -191,7 +195,7 @@ bool FVulkanViewport::DoCheckedSwapChainJob(FVulkanCommandListContext& Context,
 	int32 AttemptsPending = FVulkanPlatform::RecreateSwapchainOnFail() ? 4 : 0;
 	int32 Status = SwapChainJob(this);
 
-	while (Status < 0 && AttemptsPending > 0)
+	while (FVulkanPlatformWindowContext::CanCreateSwapchainOnDemand() && Status < 0 && AttemptsPending > 0)
 	{
 		if (Status == (int32)FVulkanSwapChain::EStatus::OutOfDate)
 		{
@@ -206,7 +210,8 @@ bool FVulkanViewport::DoCheckedSwapChainJob(FVulkanCommandListContext& Context,
 			check(0);
 		}
 
-		RecreateSwapchain(Context);
+		FVulkanPlatformWindowContext WindowContext(WindowHandle);
+		RecreateSwapchain(Context, WindowContext);
 
 		// Swapchain creation pushes some commands - flush the command buffers now to begin with a fresh state
 		Context.FlushCommands(EVulkanFlushFlags::WaitForCompletion);
@@ -674,18 +679,16 @@ bool FVulkanFramebuffer::Matches(const FRHISetRenderTargetsInfo& InRTInfo) const
 }
 
 // Tear down and recreate swapchain and related resources.
-void FVulkanViewport::RecreateSwapchain(FVulkanCommandListContext& Context)
+void FVulkanViewport::RecreateSwapchain(FVulkanCommandListContext& Context, FVulkanPlatformWindowContext& WindowContext)
 {
 	// Make sure everything is submitted and submission queue is idle
 	Context.FlushCommands(EVulkanFlushFlags::WaitForCompletion);
 
 	FScopeLock LockSwapchain(&RecreatingSwapchain);
 
-	TRHICommandList_RecursiveHazardous<FVulkanCommandListContext> RHICmdList(&Context);
-
 	FVulkanSwapChainRecreateInfo RecreateInfo = { VK_NULL_HANDLE, VK_NULL_HANDLE };
 	DestroySwapchain(&RecreateInfo);
-	CreateSwapchain(RHICmdList, &RecreateInfo);
+	CreateSwapchain(Context, &RecreateInfo, WindowContext);
 	check(RecreateInfo.Surface == VK_NULL_HANDLE);
 	check(RecreateInfo.SwapChain == VK_NULL_HANDLE);
 }
@@ -696,39 +699,46 @@ void FVulkanViewport::Tick(float DeltaTime)
 
 	if (SwapChain && FPlatformAtomics::AtomicRead(&LockToVsync) != SwapChain->DoesLockToVsync())
 	{
-		FlushRenderingCommands();
+		FVulkanPlatformWindowContext WindowContext(WindowHandle);
+
 		ENQUEUE_RENDER_COMMAND(UpdateVsync)(
-			[this](FRHICommandListImmediate& RHICmdList)
+			[this, &WindowContext](FRHICommandListImmediate& RHICmdList)
 		{
-			RecreateSwapchainFromRT(RHICmdList);
+			RecreateSwapchainFromRT(RHICmdList, PixelFormat, WindowContext);
 		});
 		FlushRenderingCommands();
 	}
 }
 
-void FVulkanViewport::Resize(uint32 InSizeX, uint32 InSizeY, bool bInIsFullscreen, EPixelFormat PreferredPixelFormat)
+void FVulkanViewport::Resize(FRHICommandListImmediate& RHICmdList, uint32 InSizeX, uint32 InSizeY, bool bInIsFullscreen, EPixelFormat PreferredPixelFormat, FVulkanPlatformWindowContext& WindowContext)
 {
-	SizeX = InSizeX;
-	SizeY = InSizeY;
-	bIsFullscreen = bInIsFullscreen;
-	PixelFormat = PreferredPixelFormat;
+	check(IsInRenderingThread());
 
-	RecreateSwapchainFromRT(FRHICommandListImmediate::Get());
+	RHICmdList.EnqueueLambda([this, InSizeX, InSizeY, bInIsFullscreen, PreferredPixelFormat](FRHICommandListBase& ExecutingCmdList)
+	{
+		SizeX = InSizeX;
+		SizeY = InSizeY;
+		bIsFullscreen = bInIsFullscreen;
+		PixelFormat = PreferredPixelFormat;
+	});
+	RecreateSwapchainFromRT(RHICmdList, PreferredPixelFormat, WindowContext);
 }
 
-void FVulkanViewport::RecreateSwapchainFromRT(FRHICommandListImmediate& RHICmdList, EPixelFormat PreferredPixelFormat, void* NewNativeWindow)
+void FVulkanViewport::RecreateSwapchainFromRT(FRHICommandListImmediate& RHICmdList, EPixelFormat PreferredPixelFormat, FVulkanPlatformWindowContext& WindowContext)
 {
 	check(IsInRenderingThread());
 
-	RHICmdList.ImmediateFlush(EImmediateFlushType::FlushRHIThread);
+	RHICmdList.EnqueueLambda([this, PreferredPixelFormat, &WindowContext](FRHICommandListBase& ExecutingCmdList)
+	{
+		FVulkanSwapChainRecreateInfo RecreateInfo = { VK_NULL_HANDLE, VK_NULL_HANDLE };
+		DestroySwapchain(&RecreateInfo);
+		PixelFormat = PreferredPixelFormat;
+		CreateSwapchain(FVulkanCommandListContext::Get(ExecutingCmdList), &RecreateInfo, WindowContext);
+		check(RecreateInfo.Surface == VK_NULL_HANDLE);
+		check(RecreateInfo.SwapChain == VK_NULL_HANDLE);
+	});
 
-	FVulkanSwapChainRecreateInfo RecreateInfo = { VK_NULL_HANDLE, VK_NULL_HANDLE};
-	DestroySwapchain(WindowHandle == NewNativeWindow ? &RecreateInfo : nullptr);
-	PixelFormat = PreferredPixelFormat;
-	WindowHandle = NewNativeWindow;
-	CreateSwapchain(RHICmdList, &RecreateInfo);
-	check(RecreateInfo.Surface == VK_NULL_HANDLE);
-	check(RecreateInfo.SwapChain == VK_NULL_HANDLE);
+	RHICmdList.ImmediateFlush(EImmediateFlushType::FlushRHIThread);
 }
 
 void FVulkanViewport::InitImages(FVulkanContextCommon& Context, TConstArrayView<VkImage> Images)
@@ -769,7 +779,7 @@ void FVulkanViewport::InitImages(FVulkanContextCommon& Context, TConstArrayView<
 	}
 }
 
-void FVulkanViewport::CreateSwapchain(FRHICommandListBase& RHICmdList, FVulkanSwapChainRecreateInfo* RecreateInfo)
+void FVulkanViewport::CreateSwapchain(FVulkanCommandListContext& Context, FVulkanSwapChainRecreateInfo* RecreateInfo, FVulkanPlatformWindowContext& WindowContext)
 {
 	// Release a previous swapchain 'dummy' and a real backbuffer if any
 	RenderingBackBuffer = nullptr;
@@ -777,15 +787,17 @@ void FVulkanViewport::CreateSwapchain(FRHICommandListBase& RHICmdList, FVulkanSw
 
 	if (SupportsStandardSwapchain())
 	{
-		uint32 DesiredNumBackBuffers = NUM_BUFFERS;
+		check(SwapChain == nullptr);
 
+		uint32 DesiredNumBackBuffers = NUM_BUFFERS;
 		TArray<VkImage> Images;
 		SwapChain = new FVulkanSwapChain(
-			FVulkanDynamicRHI::Get().Instance, *Device, WindowHandle,
+			FVulkanDynamicRHI::Get().Instance, *Device,
 			PixelFormat, SizeX, SizeY, bIsFullscreen,
 			&DesiredNumBackBuffers,
 			Images,
 			LockToVsync,
+			WindowContext,
 			RecreateInfo
 		);
 
@@ -793,14 +805,7 @@ void FVulkanViewport::CreateSwapchain(FRHICommandListBase& RHICmdList, FVulkanSw
 		BackBufferImages.SetNum(Images.Num());
 		RenderingDoneSemaphores.SetNum(Images.Num());
 
-		RHICmdList.EnqueueLambda(TEXT("FVulkanViewport::CreateSwapchain"), [VulkanViewport = this, Images](FRHICommandListBase& ExecutingCmdList) 
-		{
-			VulkanViewport->InitImages(FVulkanUploadContext::Get(ExecutingCmdList), Images);
-		});
-		if (RHICmdList.IsImmediate())
-		{
-			RHICmdList.GetAsImmediate().SubmitAndBlockUntilGPUIdle();
-		}
+		InitImages(Context, Images);
 
 		RHIBackBuffer = new FVulkanBackBuffer(*Device, this, PixelFormat, SizeX, SizeY, TexCreate_RenderTargetable | TexCreate_ShaderResource | TexCreate_ResolveTargetable);
 	}
@@ -837,7 +842,15 @@ void FVulkanViewport::CreateSwapchain(FRHICommandListBase& RHICmdList, FVulkanSw
 			.SetFlags(ETextureCreateFlags::RenderTargetable | ETextureCreateFlags::ShaderResource | ETextureCreateFlags::ResolveTargetable)
 			.SetInitialState(ERHIAccess::Present);
 
-		RenderingBackBuffer = FVulkanDynamicRHI::Get().CreateTextureInternal(RHICmdList, CreateDesc, {});
+
+		RenderingBackBuffer = FVulkanDynamicRHI::Get().CreateTextureInternal(CreateDesc, {});
+
+		FVulkanPipelineBarrier Barrier;
+
+		Barrier.AddImageLayoutTransition(RenderingBackBuffer->Image, VK_IMAGE_LAYOUT_UNDEFINED, VK_IMAGE_LAYOUT_PRESENT_SRC_KHR, FVulkanPipelineBarrier::MakeSubresourceRange(VK_IMAGE_ASPECT_COLOR_BIT, 0, 1));
+
+		Barrier.Execute(&Context.GetCommandBuffer());
+
 #if VULKAN_ENABLE_DRAW_MARKERS
 		if (Device->GetSetDebugName())
 		{
@@ -1054,15 +1067,19 @@ bool FVulkanViewport::Present(FVulkanCommandListContext& Context, FVulkanQueue*
 		}
 		else
 		{
-			// failing to do the delayacquire can only happen if we were in this mode to begin with
-			check(GVulkanDelayAcquireImage == EDelayAcquireImageType::DelayAcquire);
+			if(FVulkanPlatformWindowContext::CanCreateSwapchainOnDemand()) // on android we dont want to attempt to recreate the window when we dont have the window lock...
+			{
+				// failing to do the delayacquire can only happen if we were in this mode to begin with
+				check(GVulkanDelayAcquireImage == EDelayAcquireImageType::DelayAcquire);
 
-			UE_LOG(LogVulkanRHI, Log, TEXT("AcquireNextImage() failed due to the outdated swapchain, not even attempting to present."));
+				UE_LOG(LogVulkanRHI, Log, TEXT("AcquireNextImage() failed due to the outdated swapchain, not even attempting to present."));
 
-			RecreateSwapchain(Context);
+				FVulkanPlatformWindowContext WindowContext(GetWindowHandle());
+				RecreateSwapchain(Context, WindowContext);
 
-			// Swapchain creation pushes some commands - flush the command buffers now to begin with a fresh state
-			Context.FlushCommands(EVulkanFlushFlags::WaitForCompletion);
+				// Swapchain creation pushes some commands - flush the command buffers now to begin with a fresh state
+				Context.FlushCommands(EVulkanFlushFlags::WaitForCompletion);
+			}
 
 			// early exit
 			return (int32)FVulkanSwapChain::EStatus::Healthy;
@@ -1179,13 +1196,12 @@ void FVulkanViewport::OnSystemResolutionChanged(uint32 ResX, uint32 ResY)
 		|| (CachedOrientation == EDeviceScreenOrientation::LandscapeLeft && CurrentOrientation == EDeviceScreenOrientation::LandscapeRight))
 	{
 		check(IsInGameThread());
-
-		FlushRenderingCommands();
+		FVulkanPlatformWindowContext WindowContext(GetWindowHandle());
 
 		ENQUEUE_RENDER_COMMAND(RecreateSwapchain)(
-			[this](FRHICommandListImmediate& RHICmdList)
+			[this,&WindowContext](FRHICommandListImmediate& RHICmdList)
 			{
-				RecreateSwapchainFromRT(RHICmdList);
+				RecreateSwapchainFromRT(RHICmdList, WindowContext);
 			});
 		FlushRenderingCommands();
 	}
@@ -1223,12 +1239,12 @@ void FVulkanDynamicRHI::RHIResizeViewport(FRHIViewport* ViewportRHI, uint32 Size
 
 	if (Viewport->GetSizeXY() != FIntPoint(SizeX, SizeY) || Viewport->IsFullscreen() != bIsFullscreen)
 	{
-		FlushRenderingCommands();
+		FVulkanPlatformWindowContext WindowContext(Viewport->WindowHandle);
 
 		ENQUEUE_RENDER_COMMAND(ResizeViewport)(
-			[Viewport, SizeX, SizeY, bIsFullscreen, PreferredPixelFormat](FRHICommandListImmediate& RHICmdList)
+			[Viewport, SizeX, SizeY, bIsFullscreen, PreferredPixelFormat,&WindowContext](FRHICommandListImmediate& RHICmdList)
 			{
-				Viewport->Resize(SizeX, SizeY, bIsFullscreen, PreferredPixelFormat);
+				Viewport->Resize(RHICmdList, SizeX, SizeY, bIsFullscreen, PreferredPixelFormat, WindowContext);
 			});
 		FlushRenderingCommands();
 	}
@@ -1241,12 +1257,12 @@ void FVulkanDynamicRHI::RHIResizeViewport(FRHIViewport* ViewportRHI, uint32 Size
 
 	if (Viewport->GetSizeXY() != FIntPoint(SizeX, SizeY))
 	{
-		FlushRenderingCommands();
+		FVulkanPlatformWindowContext WindowContext(Viewport->WindowHandle);
 
 		ENQUEUE_RENDER_COMMAND(ResizeViewport)(
-			[Viewport, SizeX, SizeY, bIsFullscreen](FRHICommandListImmediate& RHICmdList)
+			[Viewport, SizeX, SizeY, bIsFullscreen, &WindowContext](FRHICommandListImmediate& RHICmdList)
 			{
-				Viewport->Resize(SizeX, SizeY, bIsFullscreen, PF_Unknown);
+				Viewport->Resize(RHICmdList, SizeX, SizeY, bIsFullscreen, PF_Unknown, WindowContext);
 			});
 		FlushRenderingCommands();
 	}
diff --git a/Engine/Source/Runtime/VulkanRHI/Private/VulkanViewport.h b/Engine/Source/Runtime/VulkanRHI/Private/VulkanViewport.h
--- a/Engine/Source/Runtime/VulkanRHI/Private/VulkanViewport.h
+++ b/Engine/Source/Runtime/VulkanRHI/Private/VulkanViewport.h
@@ -108,6 +108,10 @@ public:
 	}
 
 	VkFormat GetSwapchainImageFormat() const;
+	inline void* GetWindowHandle() { return WindowHandle; }
+
+	void DestroySwapchain(struct FVulkanSwapChainRecreateInfo* RecreateInfo);
+	void RecreateSwapchain(FVulkanCommandListContext& Context, FVulkanPlatformWindowContext& WindowContext);
 
 protected:
 	// NUM_BUFFERS don't have to match exactly as the driver can require a minimum number larger than NUM_BUFFERS. Provide some slack
@@ -144,18 +148,16 @@ protected:
 	EDeviceScreenOrientation CachedOrientation = EDeviceScreenOrientation::Unknown;
 	void OnSystemResolutionChanged(uint32 ResX, uint32 ResY);
 
-	void CreateSwapchain(FRHICommandListBase& RHICmdList, struct FVulkanSwapChainRecreateInfo* RecreateInfo);
-	void DestroySwapchain(struct FVulkanSwapChainRecreateInfo* RecreateInfo);
+	void CreateSwapchain(FVulkanCommandListContext& Context, struct FVulkanSwapChainRecreateInfo* RecreateInfo, FVulkanPlatformWindowContext& WindowContext);
 	bool TryAcquireImageIndex();
 	void InitImages(FVulkanContextCommon& Context, TConstArrayView<VkImage> Images);
 
-	void RecreateSwapchain(FVulkanCommandListContext& Context);
-	void RecreateSwapchainFromRT(FRHICommandListImmediate& RHICmdList, EPixelFormat PreferredPixelFormat, void* NewNativeWindow);
-	void RecreateSwapchainFromRT(FRHICommandListImmediate& RHICmdList)
+	void RecreateSwapchainFromRT(FRHICommandListImmediate& RHICmdList, EPixelFormat PreferredPixelFormat, FVulkanPlatformWindowContext& WindowContext);
+	void RecreateSwapchainFromRT(FRHICommandListImmediate& RHICmdList, FVulkanPlatformWindowContext& WindowContext)
 	{
-		RecreateSwapchainFromRT(RHICmdList, PixelFormat, WindowHandle);
+		RecreateSwapchainFromRT(RHICmdList, PixelFormat, WindowContext);
 	}
-	void Resize(uint32 InSizeX, uint32 InSizeY, bool bIsFullscreen, EPixelFormat PreferredPixelFormat);
+	void Resize(FRHICommandListImmediate& RHICmdList, uint32 InSizeX, uint32 InSizeY, bool bIsFullscreen, EPixelFormat PreferredPixelFormat, FVulkanPlatformWindowContext& WindowContext);
 
 	bool DoCheckedSwapChainJob(FVulkanCommandListContext& Context, TFunction<int32(FVulkanViewport*)> SwapChainJob);
 	bool SupportsStandardSwapchain();
diff --git a/Engine/Source/Runtime/VulkanRHI/Private/Windows/VulkanWindowsPlatform.cpp b/Engine/Source/Runtime/VulkanRHI/Private/Windows/VulkanWindowsPlatform.cpp
--- a/Engine/Source/Runtime/VulkanRHI/Private/Windows/VulkanWindowsPlatform.cpp
+++ b/Engine/Source/Runtime/VulkanRHI/Private/Windows/VulkanWindowsPlatform.cpp
@@ -180,12 +180,12 @@ void FVulkanWindowsPlatform::GetDeviceExtensions(FVulkanDevice* Device, FVulkanD
 														VULKAN_EXTENSION_NOT_PROMOTED, nullptr, FVulkanExtensionBase::ManuallyActivate));
 }
 
-void FVulkanWindowsPlatform::CreateSurface(void* WindowHandle, VkInstance Instance, VkSurfaceKHR* OutSurface)
+void FVulkanWindowsPlatform::CreateSurface(FVulkanPlatformWindowContext& WindowContext, VkInstance Instance, VkSurfaceKHR* OutSurface)
 {
 	VkWin32SurfaceCreateInfoKHR SurfaceCreateInfo;
 	ZeroVulkanStruct(SurfaceCreateInfo, VK_STRUCTURE_TYPE_WIN32_SURFACE_CREATE_INFO_KHR);
 	SurfaceCreateInfo.hinstance = GetModuleHandle(nullptr);
-	SurfaceCreateInfo.hwnd = (HWND)WindowHandle;
+	SurfaceCreateInfo.hwnd = (HWND)WindowContext.GetWindowHandle();
 	VERIFYVULKANRESULT(VulkanDynamicAPI::vkCreateWin32SurfaceKHR(Instance, &SurfaceCreateInfo, VULKAN_CPU_ALLOCATOR, OutSurface));
 }
 
diff --git a/Engine/Source/Runtime/VulkanRHI/Private/Windows/VulkanWindowsPlatform.h b/Engine/Source/Runtime/VulkanRHI/Private/Windows/VulkanWindowsPlatform.h
--- a/Engine/Source/Runtime/VulkanRHI/Private/Windows/VulkanWindowsPlatform.h
+++ b/Engine/Source/Runtime/VulkanRHI/Private/Windows/VulkanWindowsPlatform.h
@@ -40,6 +40,8 @@
 // and now, include the GenericPlatform class
 #include "../VulkanGenericPlatform.h"
 
+typedef FVulkanGenericPlatformWindowContext FVulkanPlatformWindowContext;
+
 class FVulkanWindowsPlatform : public FVulkanGenericPlatform
 {
 public:
@@ -52,7 +54,7 @@ public:
 	static void GetDeviceExtensions(FVulkanDevice* Device, FVulkanDeviceExtensionArray& OutExtensions);
 	static void GetDeviceLayers(TArray<const ANSICHAR*>& OutLayers) {}
 
-	static void CreateSurface(void* WindowHandle, VkInstance Instance, VkSurfaceKHR* OutSurface);
+	static void CreateSurface(FVulkanPlatformWindowContext& WindowContext, VkInstance Instance, VkSurfaceKHR* OutSurface);
 
 	static bool SupportsDeviceLocalHostVisibleWithNoPenalty(EGpuVendorId VendorId);
 
