VulkanRHI:  Add the possiblity to query supported access flags on a queue.

#rnx

diff --git a/Engine/Source/Runtime/VulkanRHI/Private/VulkanQueue.cpp b/Engine/Source/Runtime/VulkanRHI/Private/VulkanQueue.cpp
--- a/Engine/Source/Runtime/VulkanRHI/Private/VulkanQueue.cpp
+++ b/Engine/Source/Runtime/VulkanRHI/Private/VulkanQueue.cpp
@@ -416,6 +416,12 @@ void FVulkanQueue::FillSupportedStageBits()
 		VK_PIPELINE_STAGE_HOST_BIT |
 		VK_PIPELINE_STAGE_ALL_COMMANDS_BIT;
 
+	SupportedAccess =
+		VK_ACCESS_HOST_READ_BIT |
+		VK_ACCESS_HOST_WRITE_BIT |
+		VK_ACCESS_MEMORY_READ_BIT |
+		VK_ACCESS_MEMORY_WRITE_BIT;
+
 	if (VKHasAnyFlags(QueueProps.queueFlags, VK_QUEUE_GRAPHICS_BIT))
 	{
 		SupportedStages |=
@@ -429,6 +435,21 @@ void FVulkanQueue::FillSupportedStageBits()
 			VK_PIPELINE_STAGE_TRANSFER_BIT |
 			VK_PIPELINE_STAGE_ALL_GRAPHICS_BIT;
 
+		SupportedAccess |=
+			VK_ACCESS_TRANSFER_READ_BIT |
+			VK_ACCESS_TRANSFER_WRITE_BIT |
+			VK_ACCESS_INDIRECT_COMMAND_READ_BIT |
+			VK_ACCESS_UNIFORM_READ_BIT |
+			VK_ACCESS_SHADER_READ_BIT |
+			VK_ACCESS_SHADER_WRITE_BIT |
+			VK_ACCESS_INDEX_READ_BIT |
+			VK_ACCESS_VERTEX_ATTRIBUTE_READ_BIT |
+			VK_ACCESS_INPUT_ATTACHMENT_READ_BIT |
+			VK_ACCESS_COLOR_ATTACHMENT_READ_BIT |
+			VK_ACCESS_COLOR_ATTACHMENT_WRITE_BIT |
+			VK_ACCESS_DEPTH_STENCIL_ATTACHMENT_READ_BIT |
+			VK_ACCESS_DEPTH_STENCIL_ATTACHMENT_WRITE_BIT;
+
 		if (Device.GetPhysicalDeviceFeatures().Core_1_0.geometryShader)
 		{
 			SupportedStages |= VK_PIPELINE_STAGE_GEOMETRY_SHADER_BIT;
@@ -436,10 +457,12 @@ void FVulkanQueue::FillSupportedStageBits()
 		if (Device.GetOptionalExtensions().HasKHRFragmentShadingRate)
 		{
 			SupportedStages |= VK_PIPELINE_STAGE_FRAGMENT_SHADING_RATE_ATTACHMENT_BIT_KHR;
+			SupportedAccess |= VK_ACCESS_FRAGMENT_SHADING_RATE_ATTACHMENT_READ_BIT_KHR;
 		}
 		if (Device.GetOptionalExtensions().HasEXTFragmentDensityMap)
 		{
 			SupportedStages |= VK_PIPELINE_STAGE_FRAGMENT_DENSITY_PROCESS_BIT_EXT;
+			SupportedAccess |= VK_ACCESS_FRAGMENT_DENSITY_MAP_READ_BIT_EXT;
 		}
 		if (Device.GetOptionalExtensions().HasEXTMeshShader)
 		{
@@ -454,9 +477,18 @@ void FVulkanQueue::FillSupportedStageBits()
 			VK_PIPELINE_STAGE_DRAW_INDIRECT_BIT |
 			VK_PIPELINE_STAGE_TRANSFER_BIT;
 
+		SupportedAccess |=
+			VK_ACCESS_TRANSFER_READ_BIT |
+			VK_ACCESS_TRANSFER_WRITE_BIT |
+			VK_ACCESS_INDIRECT_COMMAND_READ_BIT |
+			VK_ACCESS_UNIFORM_READ_BIT |
+			VK_ACCESS_SHADER_READ_BIT |
+			VK_ACCESS_SHADER_WRITE_BIT;
+
 		if (Device.GetOptionalExtensions().HasAccelerationStructure)
 		{
 			SupportedStages |= VK_PIPELINE_STAGE_ACCELERATION_STRUCTURE_BUILD_BIT_KHR;
+			SupportedAccess |= VK_ACCESS_ACCELERATION_STRUCTURE_READ_BIT_KHR | VK_ACCESS_ACCELERATION_STRUCTURE_WRITE_BIT_KHR;
 		}
 
 		if (Device.GetOptionalExtensions().HasRayTracingPipeline)
@@ -468,6 +500,7 @@ void FVulkanQueue::FillSupportedStageBits()
 	if (VKHasAnyFlags(QueueProps.queueFlags, VK_QUEUE_TRANSFER_BIT))
 	{
 		SupportedStages |= VK_PIPELINE_STAGE_TRANSFER_BIT;
+		SupportedAccess |= VK_ACCESS_TRANSFER_READ_BIT | VK_ACCESS_TRANSFER_WRITE_BIT;
 	}
 }
 
diff --git a/Engine/Source/Runtime/VulkanRHI/Private/VulkanQueue.h b/Engine/Source/Runtime/VulkanRHI/Private/VulkanQueue.h
--- a/Engine/Source/Runtime/VulkanRHI/Private/VulkanQueue.h
+++ b/Engine/Source/Runtime/VulkanRHI/Private/VulkanQueue.h
@@ -46,26 +46,31 @@ public:
 		return QueueType;
 	}
 
-	inline uint32 GetFamilyIndex() const
+	uint32 GetFamilyIndex() const
 	{
 		return FamilyIndex;
 	}
 
-	inline uint32 GetQueueIndex() const
+	uint32 GetQueueIndex() const
 	{
 		return QueueIndex;
 	}
 
-	inline VkQueue GetHandle() const
+	VkQueue GetHandle() const
 	{
 		return Queue;
 	}
 
-	inline VkPipelineStageFlags GetSupportedStageBits() const
+	VkPipelineStageFlags GetSupportedStageBits() const
 	{
 		return SupportedStages;
 	}
 
+	VkAccessFlags GetSupportedAccessFlags() const
+	{
+		return SupportedAccess;
+	}
+
 	const VulkanRHI::FSemaphore* GetTimelineSemaphore() const
 	{
 		return TimelineSempahore;
@@ -92,7 +97,7 @@ public:
 #if RHI_NEW_GPU_PROFILER
 	UE::RHI::GPUProfiler::FQueue GetProfilerQueue() const;
 #endif
-
+	
 private:
 	void SubmitPayloads(TArrayView<FVulkanPayload*> Payloads, TMap<VkSemaphore, FBinarySemaphoreSignalInfo>& InOutSignaledSemas);
 	void Submit(TArrayView<VkSubmitInfo> InSubmitInfos, VulkanRHI::FFence* Fence);
@@ -119,7 +124,8 @@ private:
 
 	mutable FCriticalSection CS;
 	uint64 SubmitCounter;
-	VkPipelineStageFlags SupportedStages;
+	VkPipelineStageFlags SupportedStages = VK_PIPELINE_STAGE_NONE;
+	VkAccessFlags SupportedAccess = VK_ACCESS_NONE;
 
 	TQueue<FVulkanPayload*, EQueueMode::Mpsc> PendingSubmission;
 	TQueue<FVulkanPayload*, EQueueMode::Spsc> PendingInterrupt;
