VulkanRHI: Move buffer creation to a function in the device.  Add the option of using concurrent mode for all buffers when async compute is activated.  Track active queue families.

#rnx

diff --git a/Engine/Source/Runtime/VulkanRHI/Private/VulkanDescriptorSets.cpp b/Engine/Source/Runtime/VulkanRHI/Private/VulkanDescriptorSets.cpp
--- a/Engine/Source/Runtime/VulkanRHI/Private/VulkanDescriptorSets.cpp
+++ b/Engine/Source/Runtime/VulkanRHI/Private/VulkanDescriptorSets.cpp
@@ -473,11 +473,7 @@ void FVulkanBindlessDescriptorManager::Init()
 
 			// Create descriptor buffer
 			{
-				VkBufferCreateInfo BufferCreateInfo;
-				ZeroVulkanStruct(BufferCreateInfo, VK_STRUCTURE_TYPE_BUFFER_CREATE_INFO);
-				BufferCreateInfo.size = DescriptorBufferSize;
-				BufferCreateInfo.usage = BufferUsageFlags;
-				VERIFYVULKANRESULT(VulkanRHI::vkCreateBuffer(DeviceHandle, &BufferCreateInfo, VULKAN_CPU_ALLOCATOR, &InOutState.BufferHandle));
+				InOutState.BufferHandle = Device.CreateBuffer(DescriptorBufferSize, BufferUsageFlags);
 			}
 
 			// Allocate buffer memory, bind and map
diff --git a/Engine/Source/Runtime/VulkanRHI/Private/VulkanDevice.cpp b/Engine/Source/Runtime/VulkanRHI/Private/VulkanDevice.cpp
--- a/Engine/Source/Runtime/VulkanRHI/Private/VulkanDevice.cpp
+++ b/Engine/Source/Runtime/VulkanRHI/Private/VulkanDevice.cpp
@@ -21,11 +21,22 @@
 #include "VulkanDescriptorSets.h"
 #include "VulkanChunkedPipelineCache.h"
 
-TAutoConsoleVariable<int32> GRHIAllowAsyncComputeCvar(
+static TAutoConsoleVariable<int32> GRHIAllowAsyncComputeCvar(
 	TEXT("r.Vulkan.AllowAsyncCompute"),
 	0,
-	TEXT("0 to disable async compute queue (if available)")
-	TEXT("1 to allow async compute queue")
+	TEXT("0 to disable async compute queue (if available)\n")
+	TEXT("1 to allow async compute queue"),
+	ECVF_ReadOnly
+);
+
+int32 GVulkanAllowConcurrentBuffer = 1;
+static TAutoConsoleVariable<int32> GCVarAllowConcurrentBuffer(
+	TEXT("r.Vulkan.AllowConcurrentBuffer"),
+	GVulkanAllowConcurrentBuffer,
+	TEXT("When async compute is supported: \n")
+	TEXT(" 0 to use queue family ownership transfers with buffers\n")
+	TEXT(" 1 to use sharing mode concurrent with buffers"),
+	ECVF_ReadOnly
 );
 
 TAutoConsoleVariable<int32> GAllowPresentOnComputeQueue(
@@ -338,8 +349,15 @@ void FVulkanDevice::CreateDevice(TArray<const ANSICHAR*>& DeviceLayers, FVulkanD
 
 		if (VKHasAllFlags(CurrProps.queueFlags, VK_QUEUE_COMPUTE_BIT))
 		{
-			if (ComputeQueueFamilyIndex == -1 &&
-				(GRHIAllowAsyncComputeCvar.GetValueOnAnyThread() != 0 || GAllowPresentOnComputeQueue.GetValueOnAnyThread() != 0) && GfxQueueFamilyIndex != FamilyIndex)
+			// Allocate a queue for async compute if:
+ 			// - async compute queue hasn't been found already
+			// - cvars allow for a dedicated async compute queue
+			// - a new family index is available
+			// - Sync2 is available
+			if ((ComputeQueueFamilyIndex == -1) &&
+				(GRHIAllowAsyncComputeCvar.GetValueOnAnyThread() != 0 || GAllowPresentOnComputeQueue.GetValueOnAnyThread() != 0) && 
+				(GfxQueueFamilyIndex != FamilyIndex) && 
+				SupportsParallelRendering())
 			{
 				ComputeQueueFamilyIndex = FamilyIndex;
 				bIsValidQueue = true;
@@ -362,9 +380,7 @@ void FVulkanDevice::CreateDevice(TArray<const ANSICHAR*>& DeviceLayers, FVulkanD
 			continue;
 		}
 
-		int32 QueueIndex = QueueFamilyInfos.Num();
-		QueueFamilyInfos.AddZeroed(1);
-		VkDeviceQueueCreateInfo& CurrQueue = QueueFamilyInfos[QueueIndex];
+		VkDeviceQueueCreateInfo& CurrQueue = QueueFamilyInfos.AddZeroed_GetRef();
 		CurrQueue.sType = VK_STRUCTURE_TYPE_DEVICE_QUEUE_CREATE_INFO;
 		CurrQueue.queueFamilyIndex = FamilyIndex;
 		CurrQueue.queueCount = CurrProps.queueCount;
@@ -401,8 +417,11 @@ void FVulkanDevice::CreateDevice(TArray<const ANSICHAR*>& DeviceLayers, FVulkanD
 
 	FVulkanPlatform::NotifyFoundDeviceLayersAndExtensions(Gpu, DeviceLayers, DeviceExtensions);
 
+	ActiveQueueFamilies.Reserve((int32)EVulkanQueueType::Count);
+
 	// Create Graphics Queue, here we submit command buffers for execution
 	Queues[(int32)EVulkanQueueType::Graphics] = new FVulkanQueue(*this, GfxQueueFamilyIndex, EVulkanQueueType::Graphics);
+	ActiveQueueFamilies.Add(GfxQueueFamilyIndex);
 	if (ComputeQueueFamilyIndex == -1)
 	{
 		// If we didn't find a dedicated Queue, use the default one
@@ -410,11 +429,8 @@ void FVulkanDevice::CreateDevice(TArray<const ANSICHAR*>& DeviceLayers, FVulkanD
 	}
 	else
 	{
-		// Dedicated queue
-		if (GRHIAllowAsyncComputeCvar.GetValueOnAnyThread() != 0)
-		{
-			bAsyncComputeQueue = true;
-		}
+		bAsyncComputeQueue = true;
+		ActiveQueueFamilies.Add(ComputeQueueFamilyIndex);
 	}
 	Queues[(int32)EVulkanQueueType::AsyncCompute] = new FVulkanQueue(*this, ComputeQueueFamilyIndex, EVulkanQueueType::AsyncCompute);
 	if (TransferQueueFamilyIndex == -1)
@@ -422,6 +438,10 @@ void FVulkanDevice::CreateDevice(TArray<const ANSICHAR*>& DeviceLayers, FVulkanD
 		// If we didn't find a dedicated Queue, use the default one
 		TransferQueueFamilyIndex = ComputeQueueFamilyIndex;
 	}
+	else
+	{
+		ActiveQueueFamilies.Add(TransferQueueFamilyIndex);
+	}
 	Queues[(int32)EVulkanQueueType::Transfer] = new FVulkanQueue(*this, TransferQueueFamilyIndex, EVulkanQueueType::Transfer);
 
 	uint64 NumBits = QueueFamilyProps[GfxQueueFamilyIndex].timestampValidBits;
@@ -1753,3 +1773,36 @@ void FVulkanDevice::InitGlobalSamplers()
 	// Make sure if bindless handles are being generated that these samplers are the very first
 	check(!GlobalSamplers[0]->GetBindlessHandle().IsValid() || GlobalSamplers[0]->GetBindlessHandle().GetIndex() == 1);
 }
+
+VkBuffer FVulkanDevice::CreateBuffer(VkDeviceSize BufferSize, VkBufferUsageFlags BufferUsageFlags, VkBufferCreateFlags BufferCreateFlags) const
+{
+	VkBufferCreateInfo BufferCreateInfo;
+	BufferCreateInfo.sType = VK_STRUCTURE_TYPE_BUFFER_CREATE_INFO;
+	BufferCreateInfo.pNext = nullptr;
+	BufferCreateInfo.size = BufferSize;
+	BufferCreateInfo.usage = BufferUsageFlags;
+	BufferCreateInfo.flags = BufferCreateFlags;
+
+	// For descriptors buffers
+	if (GetOptionalExtensions().HasBufferDeviceAddress)
+	{
+		BufferCreateInfo.usage |= VK_BUFFER_USAGE_SHADER_DEVICE_ADDRESS_BIT;
+	}
+
+	if (bAsyncComputeQueue && GVulkanAllowConcurrentBuffer)
+	{
+		BufferCreateInfo.sharingMode = VK_SHARING_MODE_CONCURRENT;
+		BufferCreateInfo.queueFamilyIndexCount = ActiveQueueFamilies.Num();
+		BufferCreateInfo.pQueueFamilyIndices = (uint32_t*)ActiveQueueFamilies.GetData();
+	}
+	else
+	{
+		BufferCreateInfo.sharingMode = VK_SHARING_MODE_EXCLUSIVE;
+		BufferCreateInfo.queueFamilyIndexCount = 0;
+		BufferCreateInfo.pQueueFamilyIndices = nullptr;
+	}
+
+	VkBuffer BufferHandle = VK_NULL_HANDLE;
+	VERIFYVULKANRESULT(VulkanRHI::vkCreateBuffer(Device, &BufferCreateInfo, VULKAN_CPU_ALLOCATOR, &BufferHandle));
+	return BufferHandle;
+}
diff --git a/Engine/Source/Runtime/VulkanRHI/Private/VulkanDevice.h b/Engine/Source/Runtime/VulkanRHI/Private/VulkanDevice.h
--- a/Engine/Source/Runtime/VulkanRHI/Private/VulkanDevice.h
+++ b/Engine/Source/Runtime/VulkanRHI/Private/VulkanDevice.h
@@ -502,7 +502,7 @@ public:
 	void ReleaseQueryPool(FVulkanQueryPool* Pool);
 	void RemoveStaleQueryPools();
 
-	inline class FVulkanPipelineStateCacheManager* GetPipelineStateCache()
+	class FVulkanPipelineStateCacheManager* GetPipelineStateCache()
 	{
 		return PipelineStateCache;
 	}
@@ -513,17 +513,17 @@ public:
 	FVulkanCommandListContext* AcquireDeferredContext(ERHIPipeline Pipeline);
 	void ReleaseDeferredContext(FVulkanCommandListContext* InContext);
 	void VulkanSetObjectName(VkObjectType Type, uint64_t Handle, const TCHAR* Name);
-	inline const FOptionalVulkanDeviceExtensions& GetOptionalExtensions() const
+	const FOptionalVulkanDeviceExtensions& GetOptionalExtensions() const
 	{
 		return OptionalDeviceExtensions;
 	}
 
-	inline const FOptionalVulkanDeviceExtensionProperties& GetOptionalExtensionProperties() const
+	const FOptionalVulkanDeviceExtensionProperties& GetOptionalExtensionProperties() const
 	{
 		return OptionalDeviceExtensionProperties;
 	}
 
-	inline bool SupportsParallelRendering() const
+	bool SupportsParallelRendering() const
 	{
 		return OptionalDeviceExtensions.HasSeparateDepthStencilLayouts && OptionalDeviceExtensions.HasKHRSynchronization2 && OptionalDeviceExtensions.HasKHRRenderPass2;
 	}
@@ -542,7 +542,7 @@ public:
 
 	void SetupPresentQueue(VkSurfaceKHR Surface);
 
-	inline const TArray<VkQueueFamilyProperties>& GetQueueFamilyProps()
+	const TArray<VkQueueFamilyProperties>& GetQueueFamilyProps()
 	{
 		return QueueFamilyProps;
 	}
@@ -564,6 +564,13 @@ public:
 		return *GlobalSamplers[(uint32)Type];
 	}
 
+	VkBuffer CreateBuffer(VkDeviceSize BufferSize, VkBufferUsageFlags BufferUsageFlags, VkBufferCreateFlags BufferCreateFlags=0) const;
+
+	const TArray<uint32>& GetActiveQueueFamilies() const
+	{
+		return ActiveQueueFamilies;
+	}
+
 private:
 	void MapBufferFormatSupport(FPixelFormatInfo& PixelFormatInfo, EPixelFormat UEFormat, VkFormat VulkanFormat);
 	void MapImageFormatSupport(FPixelFormatInfo& PixelFormatInfo, const TArrayView<const VkFormat>& PrioritizedFormats, EPixelFormatCapabilities RequiredCapabilities);
@@ -633,6 +640,7 @@ private:
 	FVulkanQueue* PresentQueue = nullptr;  // points to an existing queue
 	bool bAsyncComputeQueue = false;
 	bool bPresentOnComputeQueue = false;
+	TArray<uint32> ActiveQueueFamilies;
 
 	EGpuVendorId VendorId = EGpuVendorId::NotQueried;
 
diff --git a/Engine/Source/Runtime/VulkanRHI/Private/VulkanMemory.cpp b/Engine/Source/Runtime/VulkanRHI/Private/VulkanMemory.cpp
--- a/Engine/Source/Runtime/VulkanRHI/Private/VulkanMemory.cpp
+++ b/Engine/Source/Runtime/VulkanRHI/Private/VulkanMemory.cpp
@@ -2778,18 +2778,7 @@ namespace VulkanRHI
 
 		// New Buffer
 		const uint32 BufferSize = FMath::Max(Size, BufferSizes[PoolSize]);
-
-		VkBuffer Buffer;
-		VkBufferCreateInfo BufferCreateInfo;
-		ZeroVulkanStruct(BufferCreateInfo, VK_STRUCTURE_TYPE_BUFFER_CREATE_INFO);
-		BufferCreateInfo.size = BufferSize;
-		BufferCreateInfo.usage = BufferUsageFlags;
-		// For descriptors buffers
-		if (Device->GetOptionalExtensions().HasBufferDeviceAddress)
-		{
-			BufferCreateInfo.usage |= VK_BUFFER_USAGE_SHADER_DEVICE_ADDRESS_BIT;
-		}
-		VERIFYVULKANRESULT(VulkanRHI::vkCreateBuffer(Device->GetInstanceHandle(), &BufferCreateInfo, VULKAN_CPU_ALLOCATOR, &Buffer));
+		VkBuffer Buffer = Device->CreateBuffer(BufferSize, BufferUsageFlags);
 
 		VkMemoryRequirements MemReqs;
 		VulkanRHI::vkGetBufferMemoryRequirements(Device->GetInstanceHandle(), Buffer, &MemReqs);
@@ -4373,12 +4362,7 @@ namespace VulkanRHI
 		FStagingBuffer* StagingBuffer = new FStagingBuffer(Device);
 		StagingBuffer->MemoryReadFlags = InMemoryReadFlags;
 		StagingBuffer->BufferSize = Size;
-
-		VkBufferCreateInfo StagingBufferCreateInfo;
-		ZeroVulkanStruct(StagingBufferCreateInfo, VK_STRUCTURE_TYPE_BUFFER_CREATE_INFO);
-		StagingBufferCreateInfo.size = Size;
-		StagingBufferCreateInfo.usage = InUsageFlags;
-		VERIFYVULKANRESULT(VulkanRHI::vkCreateBuffer(Device->GetInstanceHandle(), &StagingBufferCreateInfo, VULKAN_CPU_ALLOCATOR, &StagingBuffer->Buffer));
+		StagingBuffer->Buffer = Device->CreateBuffer(Size, InUsageFlags);
 
 		// Set minimum alignment to 16 bytes, as some buffers are used with CPU SIMD instructions
 		uint32 ForcedMinAlignment = 16u;
@@ -4957,11 +4941,7 @@ namespace VulkanRHI
 
 		// Create buffer
 		{
-			VkBufferCreateInfo BufferCreateInfo;
-			ZeroVulkanStruct(BufferCreateInfo, VK_STRUCTURE_TYPE_BUFFER_CREATE_INFO);
-			BufferCreateInfo.size = BlockSize;
-			BufferCreateInfo.usage = BufferUsage;
-			VERIFYVULKANRESULT(VulkanRHI::vkCreateBuffer(Device->GetInstanceHandle(), &BufferCreateInfo, VULKAN_CPU_ALLOCATOR, &NewBlock->Buffer));
+			NewBlock->Buffer = Device->CreateBuffer(BlockSize, BufferUsage);
 		}
 
 		// Allocate memory
diff --git a/Engine/Source/Runtime/VulkanRHI/Private/VulkanRayTracing.cpp b/Engine/Source/Runtime/VulkanRHI/Private/VulkanRayTracing.cpp
--- a/Engine/Source/Runtime/VulkanRHI/Private/VulkanRayTracing.cpp
+++ b/Engine/Source/Runtime/VulkanRHI/Private/VulkanRayTracing.cpp
@@ -1038,11 +1038,8 @@ void FVulkanRayTracingShaderTable::Commit(FVulkanCommandListContext& Context, FR
 
 				// Alloc a new Local buffer
 				{
-					VkBufferCreateInfo BufferCreateInfo;
-					ZeroVulkanStruct(BufferCreateInfo, VK_STRUCTURE_TYPE_BUFFER_CREATE_INFO);
-					BufferCreateInfo.size = Alloc.Region.size;
-					BufferCreateInfo.usage = VK_BUFFER_USAGE_TRANSFER_DST_BIT | VK_BUFFER_USAGE_SHADER_BINDING_TABLE_BIT_KHR | VK_BUFFER_USAGE_SHADER_DEVICE_ADDRESS_BIT;
-					VERIFYVULKANRESULT(VulkanRHI::vkCreateBuffer(DeviceHandle, &BufferCreateInfo, VULKAN_CPU_ALLOCATOR, &Alloc.LocalBuffer));
+					const VkBufferUsageFlags BufferUsageFlags = VK_BUFFER_USAGE_TRANSFER_DST_BIT | VK_BUFFER_USAGE_SHADER_BINDING_TABLE_BIT_KHR | VK_BUFFER_USAGE_SHADER_DEVICE_ADDRESS_BIT;
+					Alloc.LocalBuffer = Device->CreateBuffer(Alloc.Region.size, BufferUsageFlags);
 
 					const VulkanRHI::EVulkanAllocationFlags AllocFlags = VulkanRHI::EVulkanAllocationFlags::AutoBind | VulkanRHI::EVulkanAllocationFlags::Dedicated;
 					Device->GetMemoryManager().AllocateBufferMemory(Alloc.LocalAllocation, Alloc.LocalBuffer, AllocFlags, TEXT("LocalShaderTableAllocation"), RayTracingPipelineProps.shaderGroupBaseAlignment);
diff --git a/Engine/Source/Runtime/VulkanRHI/Private/VulkanTexture.cpp b/Engine/Source/Runtime/VulkanRHI/Private/VulkanTexture.cpp
--- a/Engine/Source/Runtime/VulkanRHI/Private/VulkanTexture.cpp
+++ b/Engine/Source/Runtime/VulkanRHI/Private/VulkanTexture.cpp
@@ -1426,11 +1426,7 @@ FVulkanTexture::FVulkanTexture(FVulkanDevice& InDevice, const FRHITextureCreateD
 			Size += LocalSize;
 		}
 
-		VkBufferCreateInfo BufferCreateInfo;
-		ZeroVulkanStruct(BufferCreateInfo, VK_STRUCTURE_TYPE_BUFFER_CREATE_INFO);
-		BufferCreateInfo.size = Size;
-		BufferCreateInfo.usage = VK_BUFFER_USAGE_TRANSFER_DST_BIT;
-		VERIFYVULKANRESULT(VulkanRHI::vkCreateBuffer(InDevice.GetInstanceHandle(), &BufferCreateInfo, VULKAN_CPU_ALLOCATOR, &CpuReadbackBuffer->Buffer));
+		CpuReadbackBuffer->Buffer = InDevice.CreateBuffer(Size, VK_BUFFER_USAGE_TRANSFER_DST_BIT);
 
 		// Set minimum alignment to 16 bytes, as some buffers are used with CPU SIMD instructions
 		const uint32 ForcedMinAlignment = 16u;
@@ -1622,11 +1618,11 @@ FVulkanTexture::FVulkanTexture(FVulkanDevice& InDevice, const FRHITextureCreateD
 
 			if (!EnumHasAnyFlags(InCreateDesc.Flags, TexCreate_Presentable))
 			{
-				FRHICommandList& RHICmdList = FRHICommandListExecutor::GetImmediateCommandList();
-				SetInitialImageState(RHICmdList, InitialLayout, bDoInitialClear, InCreateDesc.ClearValue, false);
-			}
+			FRHICommandList& RHICmdList = FRHICommandListExecutor::GetImmediateCommandList();
+			SetInitialImageState(RHICmdList, InitialLayout, bDoInitialClear, InCreateDesc.ClearValue, false);
 		}
 	}
+	}
 
 	const VkImageViewType ViewType = GetViewType();
 	const VkDescriptorType DescriptorType = SupportsSampling() ? VK_DESCRIPTOR_TYPE_SAMPLED_IMAGE : VK_DESCRIPTOR_TYPE_STORAGE_IMAGE;
diff --git a/Engine/Source/Runtime/VulkanRHI/Private/VulkanTransientResourceAllocator.cpp b/Engine/Source/Runtime/VulkanRHI/Private/VulkanTransientResourceAllocator.cpp
--- a/Engine/Source/Runtime/VulkanRHI/Private/VulkanTransientResourceAllocator.cpp
+++ b/Engine/Source/Runtime/VulkanRHI/Private/VulkanTransientResourceAllocator.cpp
@@ -23,13 +23,8 @@ FVulkanTransientHeap::FVulkanTransientHeap(const FInitializer& Initializer, FVul
 	}
 
 	const bool bZeroSize = false;
-	VkBufferUsageFlags BufferUsageFlags = FVulkanBuffer::UEToVKBufferUsageFlags(InDevice, UEBufferUsageFlags, bZeroSize);
-
-	VkBufferCreateInfo BufferCreateInfo;
-	ZeroVulkanStruct(BufferCreateInfo, VK_STRUCTURE_TYPE_BUFFER_CREATE_INFO);
-	BufferCreateInfo.size = Initializer.Size;
-	BufferCreateInfo.usage = BufferUsageFlags;
-	VERIFYVULKANRESULT(VulkanRHI::vkCreateBuffer(InDevice->GetInstanceHandle(), &BufferCreateInfo, VULKAN_CPU_ALLOCATOR, &VulkanBuffer));
+	const VkBufferUsageFlags BufferUsageFlags = FVulkanBuffer::UEToVKBufferUsageFlags(InDevice, UEBufferUsageFlags, bZeroSize);
+	VulkanBuffer = InDevice->CreateBuffer(Initializer.Size, BufferUsageFlags);
 
 	// Find the alignment that works for everyone
 	const uint32 MinBufferAlignment = FMath::Max<uint32>(Initializer.Alignment, VulkanRHI::FMemoryManager::CalculateBufferAlignment(*InDevice, UEBufferUsageFlags, bZeroSize));
