VulkanRHI: Add the option to use concurrent sharing mode on images when async compute is enabled.

#rnx

diff --git a/Engine/Source/Runtime/VulkanRHI/Private/VulkanTexture.cpp b/Engine/Source/Runtime/VulkanRHI/Private/VulkanTexture.cpp
--- a/Engine/Source/Runtime/VulkanRHI/Private/VulkanTexture.cpp
+++ b/Engine/Source/Runtime/VulkanRHI/Private/VulkanTexture.cpp
@@ -29,6 +29,16 @@ static FAutoConsoleVariableRef CVarVulkanDepthStencilForceStorageBit(
 	ECVF_Default
 );
 
+int32 GVulkanAllowConcurrentImage = 1;
+static TAutoConsoleVariable<int32> GCVarAllowConcurrentImage(
+	TEXT("r.Vulkan.AllowConcurrentImage"),
+	GVulkanAllowConcurrentImage,
+	TEXT("When async compute is supported: \n")
+	TEXT(" 0 to use queue family ownership transfers with images\n")
+	TEXT(" 1 to use sharing mode concurrent with images"),
+	ECVF_ReadOnly
+);
+
 extern int32 GVulkanLogDefrag;
 
 #if ENABLE_LOW_LEVEL_MEM_TRACKER
@@ -490,6 +500,17 @@ void FVulkanTexture::GenerateImageCreateInfo(
 			}
 		}
 	}
+
+	if (InDevice.HasAsyncComputeQueue() && (GVulkanAllowConcurrentImage != 0))
+	{
+		ImageCreateInfo.sharingMode = VK_SHARING_MODE_CONCURRENT;
+		ImageCreateInfo.queueFamilyIndexCount = InDevice.GetActiveQueueFamilies().Num();
+		ImageCreateInfo.pQueueFamilyIndices = (uint32_t*)InDevice.GetActiveQueueFamilies().GetData();
+	}
+	else
+	{
+		ImageCreateInfo.sharingMode = VK_SHARING_MODE_EXCLUSIVE;
+	}
 }
 
 static VkImageLayout ChooseVRSLayout()
