Allow vulkanRHI to create a viewport with no HW window associated with it.

diff --git a/Engine/Source/Runtime/VulkanRHI/Private/Android/VulkanAndroidPlatform.h b/Engine/Source/Runtime/VulkanRHI/Private/Android/VulkanAndroidPlatform.h
--- a/Engine/Source/Runtime/VulkanRHI/Private/Android/VulkanAndroidPlatform.h
+++ b/Engine/Source/Runtime/VulkanRHI/Private/Android/VulkanAndroidPlatform.h
@@ -58,7 +58,7 @@ public:
 				// the context should be locked at the beginning of the update process, this overload is used for GT initiated events.
 				check(IsInGameThread());
 				WindowContainer = AndroidWindow->GetANativeAccessor(false);
-				NativeWindow = WindowContainer->GetANativeWindow();
+				NativeWindow = WindowContainer.IsSet() ? WindowContainer->GetANativeWindow() : nullptr;
 			}
 		}
 	}
@@ -80,6 +80,10 @@ public:
 	// we dont have a locked window to create swapchains during present, android will use the 
 	// SetRHIOnReleaseWindowCallback / SetRHIOnReInitWindowCallback to create as required.
 	static bool CanCreateSwapchainOnDemand() { return false; }
+	bool IsValid() const
+	{
+		return FAndroidMisc::UseNewWindowBehavior() ? GetANativeWindow() != nullptr : true;
+	}
 
 	ANativeWindow* GetANativeWindow() const
 	{
diff --git a/Engine/Source/Runtime/VulkanRHI/Private/VulkanGenericPlatform.h b/Engine/Source/Runtime/VulkanRHI/Private/VulkanGenericPlatform.h
--- a/Engine/Source/Runtime/VulkanRHI/Private/VulkanGenericPlatform.h
+++ b/Engine/Source/Runtime/VulkanRHI/Private/VulkanGenericPlatform.h
@@ -31,6 +31,7 @@ public:
 	void* GetWindowHandle() const { return WindowHandle; }
 	// whether the platform can attempt to create a swapchains during present
 	static bool CanCreateSwapchainOnDemand() { return true; }
+	bool IsValid() const { return true; }
 protected:
 	void* WindowHandle = nullptr;
 };
diff --git a/Engine/Source/Runtime/VulkanRHI/Private/VulkanViewport.cpp b/Engine/Source/Runtime/VulkanRHI/Private/VulkanViewport.cpp
--- a/Engine/Source/Runtime/VulkanRHI/Private/VulkanViewport.cpp
+++ b/Engine/Source/Runtime/VulkanRHI/Private/VulkanViewport.cpp
@@ -147,12 +147,6 @@ FVulkanViewport::FVulkanViewport(FVulkanDevice* InDevice, void* InWindowHandle,
 
 	if (SupportsStandardSwapchain())
 	{
-		for (int32 Index = 0, NumBuffers = RenderingDoneSemaphores.Num(); Index < NumBuffers; ++Index)
-		{
-			RenderingDoneSemaphores[Index] = new VulkanRHI::FSemaphore(*InDevice);
-			RenderingDoneSemaphores[Index]->AddRef();
-		}
-
 		FCoreDelegates::OnSystemResolutionChanged.AddRaw(this, &FVulkanViewport::OnSystemResolutionChanged);
 	}
 }
@@ -789,23 +783,38 @@ void FVulkanViewport::CreateSwapchain(FVulkanCommandListContext& Context, FVulka
 	{
 		check(SwapChain == nullptr);
 
-		uint32 DesiredNumBackBuffers = NUM_BUFFERS;
-		TArray<VkImage> Images;
-		SwapChain = new FVulkanSwapChain(
-			FVulkanDynamicRHI::Get().Instance, *Device,
-			PixelFormat, SizeX, SizeY, bIsFullscreen,
-			&DesiredNumBackBuffers,
-			Images,
-			LockToVsync,
-			WindowContext,
-			RecreateInfo
-		);
+		if (WindowContext.IsValid())
+		{
+			uint32 DesiredNumBackBuffers = NUM_BUFFERS;
+			TArray<VkImage> Images;
+			SwapChain = new FVulkanSwapChain(
+				FVulkanDynamicRHI::Get().Instance, *Device,
+				PixelFormat, SizeX, SizeY, bIsFullscreen,
+				&DesiredNumBackBuffers,
+				Images,
+				LockToVsync,
+				WindowContext,
+				RecreateInfo
+			);
 
-		checkf(Images.Num() >= NUM_BUFFERS, TEXT("We wanted at least %i images, actual Num: %i"), NUM_BUFFERS, Images.Num());
-		BackBufferImages.SetNum(Images.Num());
-		RenderingDoneSemaphores.SetNum(Images.Num());
+			checkf(Images.Num() >= NUM_BUFFERS, TEXT("We wanted at least %i images, actual Num: %i"), NUM_BUFFERS, Images.Num());
 
-		InitImages(Context, Images);
+			bool bCreateSemaphores = RenderingDoneSemaphores.IsEmpty();
+			checkf(bCreateSemaphores || RenderingDoneSemaphores.Num() == Images.Num(), TEXT("CreateSwapchain, image count is not expected to change"));
+
+			BackBufferImages.SetNum(Images.Num());
+			RenderingDoneSemaphores.SetNum(Images.Num());
+			InitImages(Context, Images);
+
+			if (bCreateSemaphores)
+			{
+				for (int32 Index = 0, NumBuffers = RenderingDoneSemaphores.Num(); Index < NumBuffers; ++Index)
+				{
+					RenderingDoneSemaphores[Index] = new VulkanRHI::FSemaphore(*Device);
+					RenderingDoneSemaphores[Index]->AddRef();
+				}
+			}
+		}
 
 		RHIBackBuffer = new FVulkanBackBuffer(*Device, this, PixelFormat, SizeX, SizeY, TexCreate_RenderTargetable | TexCreate_ShaderResource | TexCreate_ResolveTargetable);
 	}
@@ -814,7 +823,7 @@ void FVulkanViewport::CreateSwapchain(FVulkanCommandListContext& Context, FVulka
 		PixelFormat = GetPixelFormatForNonDefaultSwapchain();
 		if (RecreateInfo != nullptr)
 		{
-			if(RecreateInfo->SwapChain)
+			if (RecreateInfo->SwapChain)
 			{
 				FVulkanPlatform::DestroySwapchainKHR(Device->GetInstanceHandle(), RecreateInfo->SwapChain, VULKAN_CPU_ALLOCATOR);
 				RecreateInfo->SwapChain = VK_NULL_HANDLE;
