VulkanRHI: Add a pool of events for eventual split barriers.

#rnx

diff --git a/Engine/Source/Runtime/VulkanRHI/Private/VulkanCommandBuffer.cpp b/Engine/Source/Runtime/VulkanRHI/Private/VulkanCommandBuffer.cpp
--- a/Engine/Source/Runtime/VulkanRHI/Private/VulkanCommandBuffer.cpp
+++ b/Engine/Source/Runtime/VulkanRHI/Private/VulkanCommandBuffer.cpp
@@ -114,6 +114,12 @@ void FVulkanCommandBuffer::End(FVulkanQueryPool* QueryPool)
 {
 	checkf(IsOutsideRenderPass(), TEXT("Can't End as we're inside a render pass! CmdBuffer 0x%p State=%d"), CommandBufferHandle, (int32)State);
 
+	// Reset barrier events for next use
+	for (VkEvent BarrierEvent : EndedBarrierEvents)
+	{
+		VulkanRHI::vkCmdResetEvent(GetHandle(), BarrierEvent, VK_PIPELINE_STAGE_ALL_COMMANDS_BIT);
+	}
+
 #if RHI_NEW_GPU_PROFILER
 	if (QueryPool)
 	{
@@ -242,6 +248,12 @@ void FVulkanCommandBuffer::ResetFenceStatus()
 			check(TypedDescriptorPoolSets.Num() == 0);
 		}
 
+		for (VkEvent BarrierEvent : EndedBarrierEvents)
+		{
+			Device.ReleaseBarrierEvent(BarrierEvent);
+		}
+		EndedBarrierEvents.Reset();
+
 		// Change state at the end to be safe
 		State = EState::NeedReset;
 	}
@@ -254,6 +266,18 @@ void FVulkanCommandBuffer::SetSubmitted()
 	SubmittedTime = FPlatformTime::Seconds();
 }
 
+void FVulkanCommandBuffer::BeginSplitBarrier(VkEvent BarrierEvent, const VkDependencyInfo& DependencyInfo)
+{
+	VulkanRHI::vkCmdSetEvent2KHR(GetHandle(), BarrierEvent, &DependencyInfo);
+}
+
+void FVulkanCommandBuffer::EndSplitBarrier(VkEvent BarrierEvent, const VkDependencyInfo& DependencyInfo)
+{
+	VulkanRHI::vkCmdWaitEvents2KHR(GetHandle(), 1, &BarrierEvent, &DependencyInfo);
+	EndedBarrierEvents.Add(BarrierEvent);
+}
+
+
 FVulkanCommandBufferPool::FVulkanCommandBufferPool(FVulkanDevice& InDevice, FVulkanContextCommon& InContext)
 	: Device(InDevice)
 	, Context(InContext)
@@ -334,3 +358,4 @@ FVulkanCommandBuffer* FVulkanCommandBufferPool::Create()
 	check(CmdBuffer);
 	return CmdBuffer;
 }
+
diff --git a/Engine/Source/Runtime/VulkanRHI/Private/VulkanCommandBuffer.h b/Engine/Source/Runtime/VulkanRHI/Private/VulkanCommandBuffer.h
--- a/Engine/Source/Runtime/VulkanRHI/Private/VulkanCommandBuffer.h
+++ b/Engine/Source/Runtime/VulkanRHI/Private/VulkanCommandBuffer.h
@@ -104,6 +104,10 @@ public:
 	void BeginRenderPass(const FVulkanRenderTargetLayout& Layout, class FVulkanRenderPass* RenderPass, class FVulkanFramebuffer* Framebuffer, const VkClearValue* AttachmentClearValues);
 	void EndRenderPass();
 
+	// Split barriers
+	void BeginSplitBarrier(VkEvent BarrierEvent, const VkDependencyInfo& DependencyInfo);
+	void EndSplitBarrier(VkEvent BarrierEvent, const VkDependencyInfo& DependencyInfo);
+
 	//#todo-rco: Hide this
 	FVulkanDescriptorPoolSetContainer* CurrentDescriptorPoolSetContainer = nullptr;
 
@@ -159,6 +163,8 @@ private:
 
 	uint64 LastValidTiming = 0;
 
+	TArray<VkEvent> EndedBarrierEvents;
+
 	void AcquirePoolSetContainer();
 
 	void AllocMemory();
diff --git a/Engine/Source/Runtime/VulkanRHI/Private/VulkanDevice.cpp b/Engine/Source/Runtime/VulkanRHI/Private/VulkanDevice.cpp
--- a/Engine/Source/Runtime/VulkanRHI/Private/VulkanDevice.cpp
+++ b/Engine/Source/Runtime/VulkanRHI/Private/VulkanDevice.cpp
@@ -429,7 +429,7 @@ void FVulkanDevice::CreateDevice(TArray<const ANSICHAR*>& DeviceLayers, FVulkanD
 	}
 	else
 	{
-		bAsyncComputeQueue = true;
+			bAsyncComputeQueue = true;
 		ActiveQueueFamilies.Add(ComputeQueueFamilyIndex);
 	}
 	Queues[(int32)EVulkanQueueType::AsyncCompute] = new FVulkanQueue(*this, ComputeQueueFamilyIndex, EVulkanQueueType::AsyncCompute);
@@ -1545,6 +1545,12 @@ void FVulkanDevice::Destroy()
 
 	FVulkanUploadContext::DestroyPool();
 
+	for (VkEvent Event : BarrierEvents)
+	{
+		VulkanRHI::vkDestroyEvent(Device, Event, VULKAN_CPU_ALLOCATOR);
+	}
+	BarrierEvents.Empty();
+
 	delete TempBlockAllocator;
 	TempBlockAllocator = nullptr;
 
@@ -1806,3 +1812,36 @@ VkBuffer FVulkanDevice::CreateBuffer(VkDeviceSize BufferSize, VkBufferUsageFlags
 	VERIFYVULKANRESULT(VulkanRHI::vkCreateBuffer(Device, &BufferCreateInfo, VULKAN_CPU_ALLOCATOR, &BufferHandle));
 	return BufferHandle;
 }
+
+VkEvent FVulkanDevice::GetBarrierEvent()
+{
+	check(SupportsParallelRendering());
+
+	VkEvent Handle = VK_NULL_HANDLE;
+
+	// Check if we already have one
+	{
+		FScopeLock Lock(&BarrierEventLock);
+		if (BarrierEvents.Num())
+		{
+			Handle = BarrierEvents.Pop(EAllowShrinking::No);
+		}
+	}
+
+	// Create a new handle if none were available
+	if (!Handle)
+	{
+		VkEventCreateInfo Info;
+		ZeroVulkanStruct(Info, VK_STRUCTURE_TYPE_EVENT_CREATE_INFO);
+		Info.flags = VK_EVENT_CREATE_DEVICE_ONLY_BIT;
+		VERIFYVULKANRESULT(VulkanRHI::vkCreateEvent(Device, &Info, VULKAN_CPU_ALLOCATOR, &Handle));
+	}
+
+	return Handle;
+}
+
+void FVulkanDevice::ReleaseBarrierEvent(VkEvent Handle)
+{
+	FScopeLock Lock(&BarrierEventLock);
+	BarrierEvents.Add(Handle);
+}
diff --git a/Engine/Source/Runtime/VulkanRHI/Private/VulkanDevice.h b/Engine/Source/Runtime/VulkanRHI/Private/VulkanDevice.h
--- a/Engine/Source/Runtime/VulkanRHI/Private/VulkanDevice.h
+++ b/Engine/Source/Runtime/VulkanRHI/Private/VulkanDevice.h
@@ -564,6 +564,9 @@ public:
 		return *GlobalSamplers[(uint32)Type];
 	}
 
+	VkEvent GetBarrierEvent();
+	void ReleaseBarrierEvent(VkEvent Handle);
+
 	VkBuffer CreateBuffer(VkDeviceSize BufferSize, VkBufferUsageFlags BufferUsageFlags, VkBufferCreateFlags BufferCreateFlags=0) const;
 
 	const TArray<uint32>& GetActiveQueueFamilies() const
@@ -634,6 +637,10 @@ private:
 	uint32 OcclusionQueryPoolSize = 256;
 	TStaticArray<TArray<FVulkanQueryPool*>, (int32)EVulkanQueryPoolType::Count> FreeQueryPools;
 
+	// Reusable gpu-only barrier events
+	FCriticalSection BarrierEventLock;
+	TArray<VkEvent> BarrierEvents;
+
 	uint64 TimestampValidBitsMask = 0;
 
 	TStaticArray<FVulkanQueue*, (int32)EVulkanQueueType::Count> Queues;
