Backed out so the change can be done again with the correct credits

#rnx

diff --git a/Engine/Source/Runtime/VulkanRHI/Private/VulkanTexture.cpp b/Engine/Source/Runtime/VulkanRHI/Private/VulkanTexture.cpp
--- a/Engine/Source/Runtime/VulkanRHI/Private/VulkanTexture.cpp
+++ b/Engine/Source/Runtime/VulkanRHI/Private/VulkanTexture.cpp
@@ -1835,7 +1835,6 @@ FVulkanTexture::FVulkanTexture(FVulkanDevice& InDevice, const FRHITextureCreateD
 		VulkanImage,
 		VulkanDeviceMemory,
 		SamplerYcbcrConversion,
-		HardwareBuffer
 	};
 
 	ExternalImageDeleteCallbackInfo =
