VulkanRHI: Clean up:  Rename ResetFenceStatus in command buffer.  Take reference to pool.

#rnx

diff --git a/Engine/Source/Runtime/VulkanRHI/Private/VulkanCommandBuffer.cpp b/Engine/Source/Runtime/VulkanRHI/Private/VulkanCommandBuffer.cpp
--- a/Engine/Source/Runtime/VulkanRHI/Private/VulkanCommandBuffer.cpp
+++ b/Engine/Source/Runtime/VulkanRHI/Private/VulkanCommandBuffer.cpp
@@ -14,14 +14,14 @@
 #define CMD_BUFFER_TIME_TO_WAIT_BEFORE_DELETING		10
 
 
-FVulkanCommandBuffer::FVulkanCommandBuffer(FVulkanDevice& InDevice, FVulkanCommandBufferPool* InCommandBufferPool)
+FVulkanCommandBuffer::FVulkanCommandBuffer(FVulkanDevice& InDevice, FVulkanCommandBufferPool& InCommandBufferPool)
 	: Device(InDevice)
 	, CommandBufferPool(InCommandBufferPool)
 #if RHI_NEW_GPU_PROFILER
-	, EventStream(InCommandBufferPool->GetContext().Queue.GetProfilerQueue())
+	, EventStream(InCommandBufferPool.GetContext().Queue.GetProfilerQueue())
 #endif
 {
-	FScopeLock ScopeLock(CommandBufferPool->GetCS());
+	FScopeLock ScopeLock(CommandBufferPool.GetCS());
 	AllocMemory();
 }
 
@@ -36,7 +36,7 @@ void FVulkanCommandBuffer::AllocMemory()
 	ZeroVulkanStruct(CreateCmdBufInfo, VK_STRUCTURE_TYPE_COMMAND_BUFFER_ALLOCATE_INFO);
 	CreateCmdBufInfo.level = VK_COMMAND_BUFFER_LEVEL_PRIMARY;
 	CreateCmdBufInfo.commandBufferCount = 1;
-	CreateCmdBufInfo.commandPool = CommandBufferPool->GetHandle();
+	CreateCmdBufInfo.commandPool = CommandBufferPool.GetHandle();
 
 	VERIFYVULKANRESULT(VulkanRHI::vkAllocateCommandBuffers(Device.GetInstanceHandle(), &CreateCmdBufInfo, &CommandBufferHandle));
 
@@ -63,7 +63,7 @@ void FVulkanCommandBuffer::FreeMemory()
 	// Assumes we are inside a lock for the pool
 	check(State != EState::NotAllocated);
 	check(CommandBufferHandle != VK_NULL_HANDLE);
-	VulkanRHI::vkFreeCommandBuffers(Device.GetInstanceHandle(), CommandBufferPool->GetHandle(), 1, &CommandBufferHandle);
+	VulkanRHI::vkFreeCommandBuffers(Device.GetInstanceHandle(), CommandBufferPool.GetHandle(), 1, &CommandBufferHandle);
 	CommandBufferHandle = VK_NULL_HANDLE;
 
 	DEC_DWORD_STAT(STAT_VulkanNumCmdBuffers);
@@ -142,7 +142,7 @@ void FVulkanCommandBuffer::End(FVulkanQueryPool* QueryPool)
 void FVulkanCommandBuffer::Begin(FVulkanQueryPool* QueryPool)
 {
 	{
-		FScopeLock ScopeLock(CommandBufferPool->GetCS());
+		FScopeLock ScopeLock(CommandBufferPool.GetCS());
 		if (State == EState::NeedReset)
 		{
 			VulkanRHI::vkResetCommandBuffer(CommandBufferHandle, VK_COMMAND_BUFFER_RESET_RELEASE_RESOURCES_BIT);
@@ -222,9 +222,9 @@ bool FVulkanCommandBuffer::AcquirePoolSetAndDescriptorsIfNeeded(const class FVul
 	return false;
 }
 
-void FVulkanCommandBuffer::ResetFenceStatus()
+void FVulkanCommandBuffer::Reset()
 {
-	FScopeLock ScopeLock(CommandBufferPool->GetCS());
+	FScopeLock ScopeLock(CommandBufferPool.GetCS());
 	if (State == EState::Submitted)
 	{
 		bHasPipeline = false;
@@ -261,7 +261,7 @@ void FVulkanCommandBuffer::ResetFenceStatus()
 
 void FVulkanCommandBuffer::SetSubmitted()
 {
-	FScopeLock Lock(CommandBufferPool->GetCS());
+	FScopeLock Lock(CommandBufferPool.GetCS());
 	State = FVulkanCommandBuffer::EState::Submitted;
 	SubmittedTime = FPlatformTime::Seconds();
 }
@@ -353,7 +353,7 @@ FVulkanCommandBuffer* FVulkanCommandBufferPool::Create()
 		return CmdBuffer;
 	}
 
-	FVulkanCommandBuffer* CmdBuffer = new FVulkanCommandBuffer(Device, this);
+	FVulkanCommandBuffer* CmdBuffer = new FVulkanCommandBuffer(Device, *this);
 	CmdBuffers.Add(CmdBuffer);
 	check(CmdBuffer);
 	return CmdBuffer;
diff --git a/Engine/Source/Runtime/VulkanRHI/Private/VulkanCommandBuffer.h b/Engine/Source/Runtime/VulkanRHI/Private/VulkanCommandBuffer.h
--- a/Engine/Source/Runtime/VulkanRHI/Private/VulkanCommandBuffer.h
+++ b/Engine/Source/Runtime/VulkanRHI/Private/VulkanCommandBuffer.h
@@ -30,13 +30,13 @@ protected:
 	friend class FVulkanCommandBufferPool;
 	friend class FVulkanQueue;
 
-	FVulkanCommandBuffer(FVulkanDevice& InDevice, FVulkanCommandBufferPool* InCommandBufferPool);
+	FVulkanCommandBuffer(FVulkanDevice& InDevice, FVulkanCommandBufferPool& InCommandBufferPool);
 	~FVulkanCommandBuffer();
 
 public:
 	FVulkanCommandBufferPool* GetOwner()
 	{
-		return CommandBufferPool;
+		return &CommandBufferPool;
 	}
 
 	bool IsInsideRenderPass() const
@@ -157,9 +157,9 @@ private:
 	VkCommandBuffer CommandBufferHandle = VK_NULL_HANDLE;
 	double SubmittedTime = 0.0f;
 
-	void ResetFenceStatus();
+	void Reset();
 
-	FVulkanCommandBufferPool* CommandBufferPool = nullptr;
+	FVulkanCommandBufferPool& CommandBufferPool;
 
 	uint64 LastValidTiming = 0;
 
diff --git a/Engine/Source/Runtime/VulkanRHI/Private/VulkanSubmission.cpp b/Engine/Source/Runtime/VulkanRHI/Private/VulkanSubmission.cpp
--- a/Engine/Source/Runtime/VulkanRHI/Private/VulkanSubmission.cpp
+++ b/Engine/Source/Runtime/VulkanRHI/Private/VulkanSubmission.cpp
@@ -538,7 +538,7 @@ void FVulkanDynamicRHI::CompletePayload(FVulkanPayload* Payload)
 {
 	for (FVulkanCommandBuffer* CommandBuffer : Payload->CommandBuffers)
 	{
-		CommandBuffer->ResetFenceStatus();
+		CommandBuffer->Reset();
 	}
 
 	struct FQueryResult
