VulkanRHI: Clean Up: Use references as much as possible where null pointers aren't supported/possible.  Remove unused members and duplicate params.  Prefer Context as parameter over CommandBuffer.  Start phasing out VulkanRHI::FDeviceChild.

#rnx

diff --git a/Engine/Source/Runtime/VulkanRHI/Private/VulkanCommands.cpp b/Engine/Source/Runtime/VulkanRHI/Private/VulkanCommands.cpp
--- a/Engine/Source/Runtime/VulkanRHI/Private/VulkanCommands.cpp
+++ b/Engine/Source/Runtime/VulkanRHI/Private/VulkanCommands.cpp
@@ -284,10 +284,10 @@ void FVulkanCommandListContext::RHIDispatchComputeShader(uint32 ThreadGroupCount
 
 	CommitComputeResourceTables();
 
-	FVulkanCommandBuffer& CommandBuffer = GetCommandBuffer();
-	ensure(CommandBuffer.IsOutsideRenderPass());
-	VkCommandBuffer CommandBufferHandle = CommandBuffer.GetHandle();
-	PendingComputeState->PrepareForDispatch(CommandBuffer);
+	PendingComputeState->PrepareForDispatch(*this);
+
+	ensure(GetCommandBuffer().IsOutsideRenderPass());
+	const VkCommandBuffer CommandBufferHandle = GetCommandBuffer().GetHandle();
 
 	VulkanRHI::vkCmdDispatch(CommandBufferHandle, ThreadGroupCountX, ThreadGroupCountY, ThreadGroupCountZ);
 
@@ -310,10 +310,10 @@ void FVulkanCommandListContext::RHIDispatchIndirectComputeShader(FRHIBuffer* Arg
 
 	CommitComputeResourceTables();
 
-	FVulkanCommandBuffer& CommandBuffer = GetCommandBuffer();
-	ensure(CommandBuffer.IsOutsideRenderPass());
-	VkCommandBuffer CommandBufferHandle = CommandBuffer.GetHandle();
-	PendingComputeState->PrepareForDispatch(CommandBuffer);
+	PendingComputeState->PrepareForDispatch(*this);
+
+	ensure(GetCommandBuffer().IsOutsideRenderPass());
+	const VkCommandBuffer CommandBufferHandle = GetCommandBuffer().GetHandle();
 
 	VulkanRHI::vkCmdDispatchIndirect(CommandBufferHandle, ArgumentBuffer->GetHandle(), ArgumentBuffer->GetOffset() + ArgumentOffset);
 
@@ -624,13 +624,13 @@ void FVulkanCommandListContext::RHIDrawPrimitive(uint32 BaseVertexIndex, uint32
 
 	CommitGraphicsResourceTables();
 
-	FVulkanCommandBuffer& CommandBuffer = GetCommandBuffer();
-	PendingGfxState->PrepareForDraw(CommandBuffer);
-	uint32 NumVertices = GetVertexCountForPrimitiveCount(NumPrimitives, PendingGfxState->PrimitiveType);
+	PendingGfxState->PrepareForDraw(*this);
+
+	const uint32 NumVertices = GetVertexCountForPrimitiveCount(NumPrimitives, PendingGfxState->PrimitiveType);
 
 	RHI_DRAW_CALL_STATS(PendingGfxState->PrimitiveType, NumVertices, NumPrimitives, NumInstances);
 
-	VulkanRHI::vkCmdDraw(CommandBuffer.GetHandle(), NumVertices, NumInstances, BaseVertexIndex, 0);
+	VulkanRHI::vkCmdDraw(GetCommandBuffer().GetHandle(), NumVertices, NumInstances, BaseVertexIndex, 0);
 
 #if (RHI_NEW_GPU_PROFILER == 0)
 	if (FVulkanPlatform::RegisterGPUWork() && IsImmediate())
@@ -651,9 +651,9 @@ void FVulkanCommandListContext::RHIDrawPrimitiveIndirect(FRHIBuffer* ArgumentBuf
 
 	CommitGraphicsResourceTables();
 
-	FVulkanCommandBuffer& CommandBuffer = GetCommandBuffer();
-	VkCommandBuffer CommandBufferHandle = CommandBuffer.GetHandle();
-	PendingGfxState->PrepareForDraw(CommandBuffer);
+	PendingGfxState->PrepareForDraw(*this);
+
+	const VkCommandBuffer CommandBufferHandle = GetCommandBuffer().GetHandle();
 
 	FVulkanBuffer* ArgumentBuffer = ResourceCast(ArgumentBufferRHI);
 
@@ -679,10 +679,10 @@ void FVulkanCommandListContext::RHIDrawIndexedPrimitive(FRHIBuffer* IndexBufferR
 
 	CommitGraphicsResourceTables();
 
+	PendingGfxState->PrepareForDraw(*this);
+
 	FVulkanBuffer* IndexBuffer = ResourceCast(IndexBufferRHI);
-	FVulkanCommandBuffer& CommandBuffer = GetCommandBuffer();
-	VkCommandBuffer CommandBufferHandle = CommandBuffer.GetHandle();
-	PendingGfxState->PrepareForDraw(CommandBuffer);
+	const VkCommandBuffer CommandBufferHandle = GetCommandBuffer().GetHandle();
 	VulkanRHI::vkCmdBindIndexBuffer(CommandBufferHandle, IndexBuffer->GetHandle(), IndexBuffer->GetOffset(), IndexBuffer->GetIndexType());
 
 	uint32 NumIndices = GetVertexCountForPrimitiveCount(NumPrimitives, PendingGfxState->PrimitiveType);
@@ -705,16 +705,15 @@ void FVulkanCommandListContext::RHIDrawIndexedIndirect(FRHIBuffer* IndexBufferRH
 
 	CommitGraphicsResourceTables();
 
+	PendingGfxState->PrepareForDraw(*this);
+
 	FVulkanBuffer* IndexBuffer = ResourceCast(IndexBufferRHI);
-	FVulkanCommandBuffer& CommandBuffer = GetCommandBuffer();
-	VkCommandBuffer CommandBufferHandle = CommandBuffer.GetHandle();
-	PendingGfxState->PrepareForDraw(CommandBuffer);
+	const VkCommandBuffer CommandBufferHandle = GetCommandBuffer().GetHandle();
 	VulkanRHI::vkCmdBindIndexBuffer(CommandBufferHandle, IndexBuffer->GetHandle(), IndexBuffer->GetOffset(), IndexBuffer->GetIndexType());
 
 	FVulkanBuffer* ArgumentBuffer = ResourceCast(ArgumentsBufferRHI);
 	VkDeviceSize ArgumentOffset = DrawArgumentsIndex * sizeof(VkDrawIndexedIndirectCommand);
 
-
 	VulkanRHI::vkCmdDrawIndexedIndirect(CommandBufferHandle, ArgumentBuffer->GetHandle(), ArgumentBuffer->GetOffset() + ArgumentOffset, 1, sizeof(VkDrawIndexedIndirectCommand));
 
 #if (RHI_NEW_GPU_PROFILER == 0)
@@ -734,10 +733,10 @@ void FVulkanCommandListContext::RHIDrawIndexedPrimitiveIndirect(FRHIBuffer* Inde
 
 	CommitGraphicsResourceTables();
 
+	PendingGfxState->PrepareForDraw(*this);
+
 	FVulkanBuffer* IndexBuffer = ResourceCast(IndexBufferRHI);
-	FVulkanCommandBuffer& CommandBuffer = GetCommandBuffer();
-	VkCommandBuffer CommandBufferHandle = CommandBuffer.GetHandle();
-	PendingGfxState->PrepareForDraw(CommandBuffer);
+	const VkCommandBuffer CommandBufferHandle = GetCommandBuffer().GetHandle();
 	VulkanRHI::vkCmdBindIndexBuffer(CommandBufferHandle, IndexBuffer->GetHandle(), IndexBuffer->GetOffset(), IndexBuffer->GetIndexType());
 
 	FVulkanBuffer* ArgumentBuffer = ResourceCast(ArgumentBufferRHI);
@@ -760,11 +759,9 @@ void FVulkanCommandListContext::RHIDispatchMeshShader(uint32 ThreadGroupCountX,
 
 	CommitGraphicsResourceTables();
 
-	FVulkanCommandBuffer& CommandBuffer = GetCommandBuffer();
-
-	PendingGfxState->PrepareForDraw(CommandBuffer);
+	PendingGfxState->PrepareForDraw(*this);
 
-	VulkanRHI::vkCmdDrawMeshTasksEXT(CommandBuffer.GetHandle(), ThreadGroupCountX, ThreadGroupCountY, ThreadGroupCountZ);
+	VulkanRHI::vkCmdDrawMeshTasksEXT(GetCommandBuffer().GetHandle(), ThreadGroupCountX, ThreadGroupCountY, ThreadGroupCountZ);
 
 #if (RHI_NEW_GPU_PROFILER == 0)
 	if (FVulkanPlatform::RegisterGPUWork() && IsImmediate())
@@ -780,13 +777,11 @@ void FVulkanCommandListContext::RHIDispatchIndirectMeshShader(FRHIBuffer* Argume
 
 	CommitGraphicsResourceTables();
 
-	FVulkanCommandBuffer& CommandBuffer = GetCommandBuffer();
-
-	PendingGfxState->PrepareForDraw(CommandBuffer);
+	PendingGfxState->PrepareForDraw(*this);
 
 	FVulkanBuffer* ArgumentBuffer = ResourceCast(ArgumentBufferRHI);
 
-	VulkanRHI::vkCmdDrawMeshTasksIndirectEXT(CommandBuffer.GetHandle(), ArgumentBuffer->GetHandle(), ArgumentBuffer->GetOffset() + ArgumentOffset, 1, sizeof(VkDrawMeshTasksIndirectCommandEXT));
+	VulkanRHI::vkCmdDrawMeshTasksIndirectEXT(GetCommandBuffer().GetHandle(), ArgumentBuffer->GetHandle(), ArgumentBuffer->GetOffset() + ArgumentOffset, 1, sizeof(VkDrawMeshTasksIndirectCommandEXT));
 
 #if (RHI_NEW_GPU_PROFILER == 0)
 	if (FVulkanPlatform::RegisterGPUWork() && IsImmediate())
diff --git a/Engine/Source/Runtime/VulkanRHI/Private/VulkanContext.cpp b/Engine/Source/Runtime/VulkanRHI/Private/VulkanContext.cpp
--- a/Engine/Source/Runtime/VulkanRHI/Private/VulkanContext.cpp
+++ b/Engine/Source/Runtime/VulkanRHI/Private/VulkanContext.cpp
@@ -189,8 +189,8 @@ FVulkanCommandListContext::FVulkanCommandListContext(FVulkanDynamicRHI* InRHI, F
 #endif
 
 	// Create Pending state, contains pipeline states such as current shader and etc..
-	PendingGfxState = new FVulkanPendingGfxState(&Device, *this);
-	PendingComputeState = new FVulkanPendingComputeState(&Device, *this);
+	PendingGfxState = new FVulkanPendingGfxState(Device);
+	PendingComputeState = new FVulkanPendingComputeState(Device);
 
 	GlobalUniformBuffers.AddZeroed(FUniformBufferStaticSlotRegistry::Get().GetSlotCount());
 }
diff --git a/Engine/Source/Runtime/VulkanRHI/Private/VulkanDescriptorSets.cpp b/Engine/Source/Runtime/VulkanRHI/Private/VulkanDescriptorSets.cpp
--- a/Engine/Source/Runtime/VulkanRHI/Private/VulkanDescriptorSets.cpp
+++ b/Engine/Source/Runtime/VulkanRHI/Private/VulkanDescriptorSets.cpp
@@ -598,7 +598,7 @@ void FVulkanBindlessDescriptorManager::BindDescriptorBuffers(VkCommandBuffer Com
 	}
 }
 
-void FVulkanBindlessDescriptorManager::RegisterUniformBuffers(FVulkanCommandBuffer& CommandBuffer, VkPipelineBindPoint BindPoint, const FUniformBufferDescriptorArrays& StageUBs)
+void FVulkanBindlessDescriptorManager::RegisterUniformBuffers(FVulkanCommandListContext& Context, VkPipelineBindPoint BindPoint, const FUniformBufferDescriptorArrays& StageUBs)
 {
 	checkf(bIsSupported, TEXT("Trying to RegisterUniformBuffers but bindless is not supported!"));
 
@@ -613,7 +613,7 @@ void FVulkanBindlessDescriptorManager::RegisterUniformBuffers(FVulkanCommandBuff
 	VkDeviceSize BufferOffsets[VulkanBindless::MaxNumSets];
 	FMemory::Memzero(BufferOffsets);
 
-	uint8* MappedPointer = SingleUseUBAllocator->Alloc(BlockSize, CommandBuffer.GetOwner()->GetContext(),
+	uint8* MappedPointer = SingleUseUBAllocator->Alloc(BlockSize, Context,
 		LocalBufferBindingInfo[VulkanBindless::BindlessSingleUseUniformBufferSet], BufferOffsets[VulkanBindless::BindlessSingleUseUniformBufferSet]);
 
 	for (int32 StageIndex = 0; StageIndex < ShaderStage::MaxNumStages; ++StageIndex)
@@ -651,8 +651,9 @@ void FVulkanBindlessDescriptorManager::RegisterUniformBuffers(FVulkanCommandBuff
 	}
 
 	// todo-jn: cache these states and only repeat when necessary
-	VulkanRHI::vkCmdBindDescriptorBuffersEXT(CommandBuffer.GetHandle(), VulkanBindless::MaxNumSets, LocalBufferBindingInfo);
-	VulkanRHI::vkCmdSetDescriptorBufferOffsetsEXT(CommandBuffer.GetHandle(), BindPoint, BindlessPipelineLayout, 0u, VulkanBindless::MaxNumSets, BufferIndices, BufferOffsets);
+	const VkCommandBuffer CommandBufferHandle = Context.GetCommandBuffer().GetHandle();
+	VulkanRHI::vkCmdBindDescriptorBuffersEXT(CommandBufferHandle, VulkanBindless::MaxNumSets, LocalBufferBindingInfo);
+	VulkanRHI::vkCmdSetDescriptorBufferOffsetsEXT(CommandBufferHandle, BindPoint, BindlessPipelineLayout, 0u, VulkanBindless::MaxNumSets, BufferIndices, BufferOffsets);
 }
 
 void FVulkanBindlessDescriptorManager::UpdateStatsForHandle(VkDescriptorType DescriptorType)
diff --git a/Engine/Source/Runtime/VulkanRHI/Private/VulkanDescriptorSets.h b/Engine/Source/Runtime/VulkanRHI/Private/VulkanDescriptorSets.h
--- a/Engine/Source/Runtime/VulkanRHI/Private/VulkanDescriptorSets.h
+++ b/Engine/Source/Runtime/VulkanRHI/Private/VulkanDescriptorSets.h
@@ -879,7 +879,7 @@ public:
 	void UpdateTexelBuffer(FRHIDescriptorHandle DescriptorHandle, VkDescriptorType DescriptorType, const VkBufferViewCreateInfo& ViewInfo, bool bImmediateUpdate = true);
 	void UpdateAccelerationStructure(FRHIDescriptorHandle DescriptorHandle, VkAccelerationStructureKHR AccelerationStructure, bool bImmediateUpdate = true);
 
-	void RegisterUniformBuffers(FVulkanCommandBuffer& CommandBuffer, VkPipelineBindPoint BindPoint, const FUniformBufferDescriptorArrays& StageUBs);
+	void RegisterUniformBuffers(FVulkanCommandListContext& Context, VkPipelineBindPoint BindPoint, const FUniformBufferDescriptorArrays& StageUBs);
 
 	void Unregister(FRHIDescriptorHandle DescriptorHandle);
 
diff --git a/Engine/Source/Runtime/VulkanRHI/Private/VulkanPendingState.cpp b/Engine/Source/Runtime/VulkanRHI/Private/VulkanPendingState.cpp
--- a/Engine/Source/Runtime/VulkanRHI/Private/VulkanPendingState.cpp
+++ b/Engine/Source/Runtime/VulkanRHI/Private/VulkanPendingState.cpp
@@ -339,7 +339,7 @@ void FVulkanPendingComputeState::SetUAVForUBResource(uint32 DescriptorSet, uint3
 }
 
 
-void FVulkanPendingComputeState::PrepareForDispatch(FVulkanCommandBuffer& InCommandBuffer)
+void FVulkanPendingComputeState::PrepareForDispatch(FVulkanCommandListContext& Context)
 {
 #if VULKAN_ENABLE_AGGRESSIVE_STATS
 	SCOPE_CYCLE_COUNTER(STAT_VulkanDispatchCallPrepareTime);
@@ -347,21 +347,21 @@ void FVulkanPendingComputeState::PrepareForDispatch(FVulkanCommandBuffer& InComm
 
 	check(CurrentState);
 
-	VkCommandBuffer CommandBufferHandle = InCommandBuffer.GetHandle();
-
 	if (CurrentPipeline->UsesBindless())
 	{
-		CurrentState->UpdateBindlessDescriptors(&Context, InCommandBuffer);
-		CurrentPipeline->Bind(CommandBufferHandle);
+		CurrentState->UpdateBindlessDescriptors(Context);
+		CurrentPipeline->Bind(Context.GetCommandBuffer().GetHandle());
 	}
 	else
 	{
-		const bool bHasDescriptorSets = CurrentState->UpdateDescriptorSets(&Context, &InCommandBuffer);
+		const bool bHasDescriptorSets = CurrentState->UpdateDescriptorSets(Context);
 
 		//#todo-rco: Move this to SetComputePipeline()
 #if VULKAN_ENABLE_AGGRESSIVE_STATS
 		SCOPE_CYCLE_COUNTER(STAT_VulkanPipelineBind);
 #endif
+
+		const VkCommandBuffer CommandBufferHandle = Context.GetCommandBuffer().GetHandle();
 		CurrentPipeline->Bind(CommandBufferHandle);
 		if (bHasDescriptorSets)
 		{
@@ -381,19 +381,19 @@ FVulkanPendingGfxState::~FVulkanPendingGfxState()
 	}
 }
 
-void FVulkanPendingGfxState::PrepareForDraw(FVulkanCommandBuffer& CommandBuffer)
+void FVulkanPendingGfxState::PrepareForDraw(FVulkanCommandListContext& Context)
 {
 #if VULKAN_ENABLE_AGGRESSIVE_STATS
 	SCOPE_CYCLE_COUNTER(STAT_VulkanDrawCallPrepareTime);
 #endif
 
-	check(CommandBuffer.bHasPipeline);
+	check(Context.GetCommandBuffer().bHasPipeline);
 
 	if (CurrentPipeline->UsesBindless())
 	{
 		check(!CurrentPipeline->bHasInputAttachments); // todo-jn: bindless + InputAttachments
-		UpdateDynamicStates(CommandBuffer);
-		CurrentState->UpdateBindlessDescriptors(&Context, CommandBuffer);
+		UpdateDynamicStates(Context.GetCommandBuffer());
+		CurrentState->UpdateBindlessDescriptors(Context);
 	}
 	else
 	{
@@ -404,13 +404,13 @@ void FVulkanPendingGfxState::PrepareForDraw(FVulkanCommandBuffer& CommandBuffer)
 			UpdateInputAttachments(CurrentFramebuffer);
 		}
 
-		const bool bHasDescriptorSets = CurrentState->UpdateDescriptorSets(&Context, &CommandBuffer);
+		const bool bHasDescriptorSets = CurrentState->UpdateDescriptorSets(Context);
 
-		UpdateDynamicStates(CommandBuffer);
+		UpdateDynamicStates(Context.GetCommandBuffer());
 
 		if (bHasDescriptorSets)
 		{
-			CurrentState->BindDescriptorSets(CommandBuffer.GetHandle());
+			CurrentState->BindDescriptorSets(Context.GetCommandBuffer().GetHandle());
 		}
 	}
 
@@ -483,7 +483,7 @@ void FVulkanPendingGfxState::PrepareForDraw(FVulkanCommandBuffer& CommandBuffer)
 			// Incorrect:	1, 0, 2, 3
 			// Incorrect:	0, 2, 3, 5
 			// Reordering and creation of stream binding index is done in "GenerateVertexInputStateInfo()"
-			VulkanRHI::vkCmdBindVertexBuffers(CommandBuffer.GetHandle(), 0, TemporaryIA.NumUsed, TemporaryIA.VertexBuffers, TemporaryIA.VertexOffsets);
+			VulkanRHI::vkCmdBindVertexBuffers(Context.GetCommandBuffer().GetHandle(), 0, TemporaryIA.NumUsed, TemporaryIA.VertexBuffers, TemporaryIA.VertexOffsets);
 		}
 
 		bDirtyVertexStreams = false;
diff --git a/Engine/Source/Runtime/VulkanRHI/Private/VulkanPendingState.h b/Engine/Source/Runtime/VulkanRHI/Private/VulkanPendingState.h
--- a/Engine/Source/Runtime/VulkanRHI/Private/VulkanPendingState.h
+++ b/Engine/Source/Runtime/VulkanRHI/Private/VulkanPendingState.h
@@ -17,12 +17,11 @@
 #include "VulkanPipelineState.h"
 
 // All the current compute pipeline states in use
-class FVulkanPendingComputeState : public VulkanRHI::FDeviceChild
+class FVulkanPendingComputeState
 {
 public:
-	FVulkanPendingComputeState(FVulkanDevice* InDevice, FVulkanCommandListContext& InContext)
-		: VulkanRHI::FDeviceChild(InDevice)
-		, Context(InContext)
+	FVulkanPendingComputeState(FVulkanDevice& InDevice)
+		: Device(InDevice)
 	{
 	}
 
@@ -55,7 +54,7 @@ public:
 		}
 	}
 
-	void PrepareForDispatch(FVulkanCommandBuffer& InCommandBuffer);
+	void PrepareForDispatch(FVulkanCommandListContext& Context);
 
 	inline const FVulkanComputeShader* GetCurrentShader() const
 	{
@@ -112,18 +111,17 @@ protected:
 
 	TMap<FVulkanComputePipeline*, FVulkanComputePipelineDescriptorState*> PipelineStates;
 
-	FVulkanCommandListContext& Context;
+	FVulkanDevice& Device;
 
 	friend class FVulkanCommandListContext;
 };
 
 // All the current gfx pipeline states in use
-class FVulkanPendingGfxState : public VulkanRHI::FDeviceChild
+class FVulkanPendingGfxState
 {
 public:
-	FVulkanPendingGfxState(FVulkanDevice* InDevice, FVulkanCommandListContext& InContext)
-		: VulkanRHI::FDeviceChild(InDevice)
-		, Context(InContext)
+	FVulkanPendingGfxState(FVulkanDevice& InDevice)
+		: Device(InDevice)
 	{
 		Reset();
 	}
@@ -283,7 +281,7 @@ public:
 		CurrentState->SetPackedGlobalShaderParameter(Stage, BufferIndex, Offset, NumBytes, NewValue);
 	}
 
-	void PrepareForDraw(FVulkanCommandBuffer& CommandBuffer);
+	void PrepareForDraw(FVulkanCommandListContext& Context);
 
 	bool SetGfxPipeline(FVulkanRHIGraphicsPipelineState* InGfxPipeline, bool bForceReset);
 
@@ -341,7 +339,7 @@ protected:
 	void InternalUpdateDynamicStates(FVulkanCommandBuffer& CommandBuffer);
 	void UpdateInputAttachments(FVulkanFramebuffer* Framebuffer);
 
-	FVulkanCommandListContext& Context;
+	FVulkanDevice& Device;
 
 	friend class FVulkanCommandListContext;
 };
diff --git a/Engine/Source/Runtime/VulkanRHI/Private/VulkanPipelineState.cpp b/Engine/Source/Runtime/VulkanRHI/Private/VulkanPipelineState.cpp
--- a/Engine/Source/Runtime/VulkanRHI/Private/VulkanPipelineState.cpp
+++ b/Engine/Source/Runtime/VulkanRHI/Private/VulkanPipelineState.cpp
@@ -38,7 +38,7 @@ static bool ShouldAlwaysWriteDescriptors()
 #endif
 }
 
-FVulkanComputePipelineDescriptorState::FVulkanComputePipelineDescriptorState(FVulkanDevice* InDevice, FVulkanComputePipeline* InComputePipeline)
+FVulkanComputePipelineDescriptorState::FVulkanComputePipelineDescriptorState(FVulkanDevice& InDevice, FVulkanComputePipeline* InComputePipeline)
 	: FVulkanCommonPipelineDescriptorState(InDevice, ShaderStage::NumComputeStages, InComputePipeline->UsesBindless())
 	, PackedUniformBuffersMask(0)
 	, PackedUniformBuffersDirty(0)
@@ -59,7 +59,7 @@ FVulkanComputePipelineDescriptorState::FVulkanComputePipelineDescriptorState(FVu
 	for (const FVulkanShaderHeader::FGlobalSamplerInfo& GlobalSamplerInfo : CodeHeader.GlobalSamplerInfos)
 	{
 		checkSlow(!bUseBindless);
-		SetSamplerState(ShaderStage::Compute, GlobalSamplerInfo.BindingIndex, &Device->GetGlobalSamplers(GlobalSamplerInfo.Type));
+		SetSamplerState(ShaderStage::Compute, GlobalSamplerInfo.BindingIndex, &Device.GetGlobalSamplers(GlobalSamplerInfo.Type));
 	}
 
 	ensure(DSWriter.Num() == 0 || DSWriter.Num() == 1);
@@ -97,7 +97,7 @@ void FVulkanCommonPipelineDescriptorState::CreateDescriptorWriteInfos()
 	check(DSWriter.Num() == 0);
 	DSWriter.AddDefaulted(MaxNumSets);
 
-	const FVulkanSamplerState& DefaultSampler = Device->GetDefaultSampler();
+	const FVulkanSamplerState& DefaultSampler = Device.GetDefaultSampler();
 	const FVulkanView::FTextureView& DefaultImageView = ResourceCast(GBlackTexture->TextureRHI)->DefaultView->GetTextureView();
 
 	FVulkanHashableDescriptorInfo* CurrentHashableDescriptorInfo = nullptr;
@@ -162,15 +162,15 @@ void FVulkanCommonPipelineDescriptorState::CreateDescriptorWriteInfos()
 	DescriptorSetHandles.AddZeroed(MaxNumSets);
 }
 
-static inline VulkanRHI::FVulkanAllocation UpdatePackedUniformBuffers(const FPackedUniformBuffers& PackedUniformBuffers, FVulkanCommandListContext* InContext, FVulkanCommandBuffer* InCmdBuffer)
+static inline VulkanRHI::FVulkanAllocation UpdatePackedUniformBuffers(const FPackedUniformBuffers& PackedUniformBuffers, FVulkanCommandListContext& InContext)
 {
 	const FPackedUniformBuffers::FPackedBuffer& StagedUniformBuffer = PackedUniformBuffers.GetBuffer();
 
 	const uint32 UBSize = (uint32)StagedUniformBuffer.Num();
-	const uint32 UBAlign = (uint32)InCmdBuffer->Device.GetLimits().minUniformBufferOffsetAlignment;
+	const uint32 UBAlign = (uint32)InContext.Device.GetLimits().minUniformBufferOffsetAlignment;
 
 	VulkanRHI::FVulkanAllocation TempAllocation;
-	uint8* MappedPointer = InContext->Device.GetTempBlockAllocator().Alloc(UBSize, UBAlign, *InContext, TempAllocation);
+	uint8* MappedPointer = InContext.Device.GetTempBlockAllocator().Alloc(UBSize, UBAlign, InContext, TempAllocation);
 
 	FMemory::Memcpy(MappedPointer, StagedUniformBuffer.GetData(), UBSize);
 
@@ -179,7 +179,7 @@ static inline VulkanRHI::FVulkanAllocation UpdatePackedUniformBuffers(const FPac
 
 
 template<bool bUseDynamicGlobalUBs>
-bool FVulkanComputePipelineDescriptorState::InternalUpdateDescriptorSets(FVulkanCommandListContext* CmdListContext, FVulkanCommandBuffer* CmdBuffer)
+bool FVulkanComputePipelineDescriptorState::InternalUpdateDescriptorSets(FVulkanCommandListContext& Context)
 {
 #if VULKAN_ENABLE_AGGRESSIVE_STATS
 	SCOPE_CYCLE_COUNTER(STAT_VulkanUpdateDescriptorSets);
@@ -196,14 +196,14 @@ bool FVulkanComputePipelineDescriptorState::InternalUpdateDescriptorSets(FVulkan
 #if VULKAN_ENABLE_AGGRESSIVE_STATS
 		SCOPE_CYCLE_COUNTER(STAT_VulkanApplyPackedUniformBuffers);
 #endif
-		SubmitPackedUniformBuffers<bUseDynamicGlobalUBs>(DSWriter[ShaderStage::Compute], UpdatePackedUniformBuffers(PackedUniformBuffers, CmdListContext, CmdBuffer));
+		SubmitPackedUniformBuffers<bUseDynamicGlobalUBs>(DSWriter[ShaderStage::Compute], UpdatePackedUniformBuffers(PackedUniformBuffers, Context));
 		PackedUniformBuffersDirty = 0;
 	}
 
 	// We are not using UseVulkanDescriptorCache() for compute pipelines
 	// Compute tend to use volatile resources that polute descriptor cache
 
-	if (!CmdBuffer->AcquirePoolSetAndDescriptorsIfNeeded(*DescriptorSetsLayout, true, DescriptorSetHandles.GetData()))
+	if (!Context.GetCommandBuffer().AcquirePoolSetAndDescriptorsIfNeeded(*DescriptorSetsLayout, true, DescriptorSetHandles.GetData()))
 	{
 		return false;
 	}
@@ -223,13 +223,13 @@ bool FVulkanComputePipelineDescriptorState::InternalUpdateDescriptorSets(FVulkan
 		INC_DWORD_STAT(STAT_VulkanNumDescSets);
 		SCOPE_CYCLE_COUNTER(STAT_VulkanVkUpdateDS);
 #endif
-		VulkanRHI::vkUpdateDescriptorSets(Device->GetInstanceHandle(), DSWriteContainer.DescriptorWrites.Num(), DSWriteContainer.DescriptorWrites.GetData(), 0, nullptr);
+		VulkanRHI::vkUpdateDescriptorSets(Device.GetInstanceHandle(), DSWriteContainer.DescriptorWrites.Num(), DSWriteContainer.DescriptorWrites.GetData(), 0, nullptr);
 	}
 
 	return true;
 }
 
-void FVulkanComputePipelineDescriptorState::UpdateBindlessDescriptors(FVulkanCommandListContext* CmdListContext, FVulkanCommandBuffer& CommandBuffer)
+void FVulkanComputePipelineDescriptorState::UpdateBindlessDescriptors(FVulkanCommandListContext& Context)
 {
 	check(bUseBindless);
 
@@ -252,10 +252,10 @@ void FVulkanComputePipelineDescriptorState::UpdateBindlessDescriptors(FVulkanCom
 		const FPackedUniformBuffers::FPackedBuffer& StagedUniformBuffer = PackedUniformBuffers.GetBuffer();
 		const int32 UBSize = StagedUniformBuffer.Num();
 		const int32 BindingIndex = 0;
-		const VkDeviceSize UBOffsetAlignment = Device->GetLimits().minUniformBufferOffsetAlignment;
+		const VkDeviceSize UBOffsetAlignment = Device.GetLimits().minUniformBufferOffsetAlignment;
 
 		VulkanRHI::FVulkanAllocation TempAllocation;
-		uint8* MappedPointer = Device->GetTempBlockAllocator().Alloc(UBSize, UBOffsetAlignment, *CmdListContext, TempAllocation, &DescriptorAddressInfos[BindingIndex]);
+		uint8* MappedPointer = Device.GetTempBlockAllocator().Alloc(UBSize, UBOffsetAlignment, Context, TempAllocation, &DescriptorAddressInfos[BindingIndex]);
 		FMemory::Memcpy(MappedPointer, StagedUniformBuffer.GetData(), UBSize);
 
 		PackedUniformBuffersDirty = 0;
@@ -280,7 +280,7 @@ void FVulkanComputePipelineDescriptorState::UpdateBindlessDescriptors(FVulkanCom
 		VkBufferDeviceAddressInfo BufferInfo;
 		ZeroVulkanStruct(BufferInfo, VK_STRUCTURE_TYPE_BUFFER_DEVICE_ADDRESS_INFO);
 		BufferInfo.buffer = WriteDescriptorSet.pBufferInfo->buffer;
-		VkDeviceAddress BufferAddress = VulkanRHI::vkGetBufferDeviceAddressKHR(Device->GetInstanceHandle(), &BufferInfo);
+		VkDeviceAddress BufferAddress = VulkanRHI::vkGetBufferDeviceAddressKHR(Device.GetInstanceHandle(), &BufferInfo);
 
 		DescriptorAddressInfo.sType = VK_STRUCTURE_TYPE_DESCRIPTOR_ADDRESS_INFO_EXT;
 		DescriptorAddressInfo.address = BufferAddress + WriteDescriptorSet.pBufferInfo->offset;
@@ -288,10 +288,10 @@ void FVulkanComputePipelineDescriptorState::UpdateBindlessDescriptors(FVulkanCom
 	}
 
 	// Send to descriptor manager
-	Device->GetBindlessDescriptorManager()->RegisterUniformBuffers(CommandBuffer, VK_PIPELINE_BIND_POINT_COMPUTE, StageUBs);
+	Device.GetBindlessDescriptorManager()->RegisterUniformBuffers(Context, VK_PIPELINE_BIND_POINT_COMPUTE, StageUBs);
 }
 
-FVulkanGraphicsPipelineDescriptorState::FVulkanGraphicsPipelineDescriptorState(FVulkanDevice* InDevice, FVulkanRHIGraphicsPipelineState* InGfxPipeline)
+FVulkanGraphicsPipelineDescriptorState::FVulkanGraphicsPipelineDescriptorState(FVulkanDevice& InDevice, FVulkanRHIGraphicsPipelineState* InGfxPipeline)
 	: FVulkanCommonPipelineDescriptorState(InDevice, ShaderStage::NumGraphicsStages, InGfxPipeline->UsesBindless())
 	, GfxPipeline(InGfxPipeline)
 {
@@ -304,7 +304,7 @@ FVulkanGraphicsPipelineDescriptorState::FVulkanGraphicsPipelineDescriptorState(F
 
 	UsedSetsMask = 0;
 
-	const FVulkanShaderFactory& ShaderFactory = Device->GetShaderFactory();
+	const FVulkanShaderFactory& ShaderFactory = Device.GetShaderFactory();
 
 	TStaticArray<const FVulkanShaderHeader*, ShaderStage::NumGraphicsStages> StageHeaders(InPlace, nullptr);
 
@@ -380,14 +380,14 @@ FVulkanGraphicsPipelineDescriptorState::FVulkanGraphicsPipelineDescriptorState(F
 			for (const FVulkanShaderHeader::FGlobalSamplerInfo& GlobalSamplerInfo : CodeHeader->GlobalSamplerInfos)
 			{
 				checkSlow(!bUseBindless);
-				SetSamplerState(StageIndex, GlobalSamplerInfo.BindingIndex, &Device->GetGlobalSamplers(GlobalSamplerInfo.Type));
+				SetSamplerState(StageIndex, GlobalSamplerInfo.BindingIndex, &Device.GetGlobalSamplers(GlobalSamplerInfo.Type));
 			}
 		}
 	}
 }
 
 template<bool bUseDynamicGlobalUBs>
-bool FVulkanGraphicsPipelineDescriptorState::InternalUpdateDescriptorSets(FVulkanCommandListContext* CmdListContext, FVulkanCommandBuffer* CmdBuffer)
+bool FVulkanGraphicsPipelineDescriptorState::InternalUpdateDescriptorSets(FVulkanCommandListContext& Context)
 {
 #if VULKAN_ENABLE_AGGRESSIVE_STATS
 	SCOPE_CYCLE_COUNTER(STAT_VulkanUpdateDescriptorSets);
@@ -408,7 +408,7 @@ bool FVulkanGraphicsPipelineDescriptorState::InternalUpdateDescriptorSets(FVulka
 		{
 			if (PackedUniformBuffersDirty[Stage] != 0)
 			{
-				MarkDirty(SubmitPackedUniformBuffers<bUseDynamicGlobalUBs>(DSWriter[Stage], UpdatePackedUniformBuffers(PackedUniformBuffers[Stage], CmdListContext, CmdBuffer)));
+				MarkDirty(SubmitPackedUniformBuffers<bUseDynamicGlobalUBs>(DSWriter[Stage], UpdatePackedUniformBuffers(PackedUniformBuffers[Stage], Context)));
 				PackedUniformBuffersDirty[Stage] = 0;
 			}
 		}
@@ -418,7 +418,7 @@ bool FVulkanGraphicsPipelineDescriptorState::InternalUpdateDescriptorSets(FVulka
 	{
 		if (bIsResourcesDirty)
 		{
-			Device->GetDescriptorSetCache().GetDescriptorSets(GetDSetsKey(), *DescriptorSetsLayout, DSWriter, DescriptorSetHandles.GetData());
+			Device.GetDescriptorSetCache().GetDescriptorSets(GetDSetsKey(), *DescriptorSetsLayout, DSWriter, DescriptorSetHandles.GetData());
 			bIsResourcesDirty = false;
 		}
 	}
@@ -427,7 +427,7 @@ bool FVulkanGraphicsPipelineDescriptorState::InternalUpdateDescriptorSets(FVulka
 		const bool bNeedsWrite = (bIsResourcesDirty || ShouldAlwaysWriteDescriptors());
 
 		// Allocate sets based on what changed
-		if (CmdBuffer->AcquirePoolSetAndDescriptorsIfNeeded(*DescriptorSetsLayout, bNeedsWrite, DescriptorSetHandles.GetData()))
+		if (Context.GetCommandBuffer().AcquirePoolSetAndDescriptorsIfNeeded(*DescriptorSetsLayout, bNeedsWrite, DescriptorSetHandles.GetData()))
 		{
 			uint32 RemainingSetsMask = UsedSetsMask;
 			uint32 Set = 0;
@@ -453,7 +453,7 @@ bool FVulkanGraphicsPipelineDescriptorState::InternalUpdateDescriptorSets(FVulka
 			INC_DWORD_STAT_BY(STAT_VulkanNumDescSets, NumSets);
 			SCOPE_CYCLE_COUNTER(STAT_VulkanVkUpdateDS);
 	#endif
-			VulkanRHI::vkUpdateDescriptorSets(Device->GetInstanceHandle(), DSWriteContainer.DescriptorWrites.Num(), DSWriteContainer.DescriptorWrites.GetData(), 0, nullptr);
+			VulkanRHI::vkUpdateDescriptorSets(Device.GetInstanceHandle(), DSWriteContainer.DescriptorWrites.Num(), DSWriteContainer.DescriptorWrites.GetData(), 0, nullptr);
 
 			bIsResourcesDirty = false;
 		}
@@ -462,7 +462,7 @@ bool FVulkanGraphicsPipelineDescriptorState::InternalUpdateDescriptorSets(FVulka
 	return true;
 }
 
-void FVulkanGraphicsPipelineDescriptorState::UpdateBindlessDescriptors(FVulkanCommandListContext* CmdListContext, FVulkanCommandBuffer& CommandBuffer)
+void FVulkanGraphicsPipelineDescriptorState::UpdateBindlessDescriptors(FVulkanCommandListContext& Context)
 {
 	check(bUseBindless);
 
@@ -470,7 +470,7 @@ void FVulkanGraphicsPipelineDescriptorState::UpdateBindlessDescriptors(FVulkanCo
 	check(DSWriteContainer.DescriptorBufferInfo.Num() == DSWriteContainer.DescriptorWrites.Num());
 	check(DSWriteContainer.DescriptorImageInfo.Num() == 0);
 
-	const VkDeviceSize UBOffsetAlignment = Device->GetLimits().minUniformBufferOffsetAlignment;
+	const VkDeviceSize UBOffsetAlignment = Device.GetLimits().minUniformBufferOffsetAlignment;
 
 	FVulkanBindlessDescriptorManager::FUniformBufferDescriptorArrays StageUBs;
 
@@ -502,7 +502,7 @@ void FVulkanGraphicsPipelineDescriptorState::UpdateBindlessDescriptors(FVulkanCo
 				const int32 BindingIndex = 0;
 
 				VulkanRHI::FVulkanAllocation TempAllocation;
-				uint8* MappedPointer = Device->GetTempBlockAllocator().Alloc(UBSize, UBOffsetAlignment, *CmdListContext, TempAllocation, &DescriptorAddressInfos[BindingIndex]);
+				uint8* MappedPointer = Device.GetTempBlockAllocator().Alloc(UBSize, UBOffsetAlignment, Context, TempAllocation, &DescriptorAddressInfos[BindingIndex]);
 				FMemory::Memcpy(MappedPointer, StagedUniformBuffer.GetData(), UBSize);
 
 				PackedUniformBuffersDirty[Stage] = 0;
@@ -524,7 +524,7 @@ void FVulkanGraphicsPipelineDescriptorState::UpdateBindlessDescriptors(FVulkanCo
 				VkBufferDeviceAddressInfo BufferInfo;
 				ZeroVulkanStruct(BufferInfo, VK_STRUCTURE_TYPE_BUFFER_DEVICE_ADDRESS_INFO);
 				BufferInfo.buffer = WriteDescriptorSet.pBufferInfo->buffer;
-				VkDeviceAddress BufferAddress = VulkanRHI::vkGetBufferDeviceAddressKHR(Device->GetInstanceHandle(), &BufferInfo);
+				VkDeviceAddress BufferAddress = VulkanRHI::vkGetBufferDeviceAddressKHR(Device.GetInstanceHandle(), &BufferInfo);
 
 				DescriptorAddressInfo.sType = VK_STRUCTURE_TYPE_DESCRIPTOR_ADDRESS_INFO_EXT;
 				DescriptorAddressInfo.address = BufferAddress + WriteDescriptorSet.pBufferInfo->offset;
@@ -534,7 +534,7 @@ void FVulkanGraphicsPipelineDescriptorState::UpdateBindlessDescriptors(FVulkanCo
 	}
 
 	// Send to descriptor manager
-	Device->GetBindlessDescriptorManager()->RegisterUniformBuffers(CommandBuffer, VK_PIPELINE_BIND_POINT_GRAPHICS, StageUBs);
+	Device.GetBindlessDescriptorManager()->RegisterUniformBuffers(Context, VK_PIPELINE_BIND_POINT_GRAPHICS, StageUBs);
 }
 
 
@@ -620,10 +620,10 @@ void FVulkanCommandListContext::RHISetComputePipelineState(FRHIComputePipelineSt
 }
 
 
-template bool FVulkanGraphicsPipelineDescriptorState::InternalUpdateDescriptorSets<true>(FVulkanCommandListContext* CmdListContext, FVulkanCommandBuffer* CmdBuffer);
-template bool FVulkanGraphicsPipelineDescriptorState::InternalUpdateDescriptorSets<false>(FVulkanCommandListContext* CmdListContext, FVulkanCommandBuffer* CmdBuffer);
-template bool FVulkanComputePipelineDescriptorState::InternalUpdateDescriptorSets<true>(FVulkanCommandListContext* CmdListContext, FVulkanCommandBuffer* CmdBuffer);
-template bool FVulkanComputePipelineDescriptorState::InternalUpdateDescriptorSets<false>(FVulkanCommandListContext* CmdListContext, FVulkanCommandBuffer* CmdBuffer);
+template bool FVulkanGraphicsPipelineDescriptorState::InternalUpdateDescriptorSets<true>(FVulkanCommandListContext& Context);
+template bool FVulkanGraphicsPipelineDescriptorState::InternalUpdateDescriptorSets<false>(FVulkanCommandListContext& Context);
+template bool FVulkanComputePipelineDescriptorState::InternalUpdateDescriptorSets<true>(FVulkanCommandListContext& Context);
+template bool FVulkanComputePipelineDescriptorState::InternalUpdateDescriptorSets<false>(FVulkanCommandListContext& Context);
 
 
 void FVulkanDescriptorSetWriter::CheckAllWritten()
diff --git a/Engine/Source/Runtime/VulkanRHI/Private/VulkanPipelineState.h b/Engine/Source/Runtime/VulkanRHI/Private/VulkanPipelineState.h
--- a/Engine/Source/Runtime/VulkanRHI/Private/VulkanPipelineState.h
+++ b/Engine/Source/Runtime/VulkanRHI/Private/VulkanPipelineState.h
@@ -19,11 +19,11 @@ extern TAutoConsoleVariable<int32> GDynamicGlobalUBs;
 
 
 // Common Pipeline state
-class FVulkanCommonPipelineDescriptorState : public VulkanRHI::FDeviceChild
+class FVulkanCommonPipelineDescriptorState
 {
 public:
-	FVulkanCommonPipelineDescriptorState(FVulkanDevice* InDevice, uint32 InMaxNumSets, bool InUsesBindless)
-		: VulkanRHI::FDeviceChild(InDevice)
+	FVulkanCommonPipelineDescriptorState(FVulkanDevice& InDevice, uint32 InMaxNumSets, bool InUsesBindless)
+		: Device(InDevice)
 		, MaxNumSets(InMaxNumSets)
 		, bUseBindless(InUsesBindless)
 	{
@@ -148,6 +148,8 @@ protected:
 
 	void CreateDescriptorWriteInfos();
 
+	FVulkanDevice& Device;
+
 	const uint32 MaxNumSets;
 
 	//#todo-rco: Won't work multithreaded!
@@ -178,7 +180,7 @@ protected:
 class FVulkanComputePipelineDescriptorState : public FVulkanCommonPipelineDescriptorState
 {
 public:
-	FVulkanComputePipelineDescriptorState(FVulkanDevice* InDevice, FVulkanComputePipeline* InComputePipeline);
+	FVulkanComputePipelineDescriptorState(FVulkanDevice& InDevice, FVulkanComputePipeline* InComputePipeline);
 	virtual ~FVulkanComputePipelineDescriptorState()
 	{
 		ComputePipeline->Release();
@@ -196,22 +198,22 @@ public:
 		PackedUniformBuffers.SetPackedGlobalParameter(ByteOffset, NumBytes, NewValue, PackedUniformBuffersDirty);
 	}
 
-	bool UpdateDescriptorSets(FVulkanCommandListContext* CmdListContext, FVulkanCommandBuffer* CmdBuffer)
+	bool UpdateDescriptorSets(FVulkanCommandListContext& Context)
 	{
 		check(!bUseBindless);
 
 		const bool bUseDynamicGlobalUBs = (GDynamicGlobalUBs->GetInt() > 0);
 		if (bUseDynamicGlobalUBs)
 		{
-			return InternalUpdateDescriptorSets<true>(CmdListContext, CmdBuffer);
+			return InternalUpdateDescriptorSets<true>(Context);
 		}
 		else
 		{
-			return InternalUpdateDescriptorSets<false>(CmdListContext, CmdBuffer);
+			return InternalUpdateDescriptorSets<false>(Context);
 		}
 	}
 
-	void UpdateBindlessDescriptors(FVulkanCommandListContext* CmdListContext, FVulkanCommandBuffer& CommandBuffer);
+	void UpdateBindlessDescriptors(FVulkanCommandListContext& Context);
 
 	inline void BindDescriptorSets(VkCommandBuffer CmdBuffer)
 	{
@@ -226,7 +228,7 @@ protected:
 	FVulkanComputePipeline* ComputePipeline;
 
 	template<bool bUseDynamicGlobalUBs>
-	bool InternalUpdateDescriptorSets(FVulkanCommandListContext* CmdListContext, FVulkanCommandBuffer* CmdBuffer);
+	bool InternalUpdateDescriptorSets(FVulkanCommandListContext& Context);
 
 	friend class FVulkanPendingComputeState;
 	friend class FVulkanCommandListContext;
@@ -235,7 +237,7 @@ protected:
 class FVulkanGraphicsPipelineDescriptorState : public FVulkanCommonPipelineDescriptorState
 {
 public:
-	FVulkanGraphicsPipelineDescriptorState(FVulkanDevice* InDevice, FVulkanRHIGraphicsPipelineState* InGfxPipeline);
+	FVulkanGraphicsPipelineDescriptorState(FVulkanDevice& InDevice, FVulkanRHIGraphicsPipelineState* InGfxPipeline);
 	virtual ~FVulkanGraphicsPipelineDescriptorState()
 	{
 		GfxPipeline->Release();
@@ -247,22 +249,22 @@ public:
 		PackedUniformBuffers[Stage].SetPackedGlobalParameter(ByteOffset, NumBytes, NewValue, PackedUniformBuffersDirty[Stage]);
 	}
 
-	bool UpdateDescriptorSets(FVulkanCommandListContext* CmdListContext, FVulkanCommandBuffer* CmdBuffer)
+	bool UpdateDescriptorSets(FVulkanCommandListContext& Context)
 	{
 		check(!bUseBindless);
 
 		const bool bUseDynamicGlobalUBs = (GDynamicGlobalUBs->GetInt() > 0);
 		if (bUseDynamicGlobalUBs)
 		{
-			return InternalUpdateDescriptorSets<true>(CmdListContext, CmdBuffer);
+			return InternalUpdateDescriptorSets<true>(Context);
 		}
 		else
 		{
-			return InternalUpdateDescriptorSets<false>(CmdListContext, CmdBuffer);
+			return InternalUpdateDescriptorSets<false>(Context);
 		}
 	}
 
-	void UpdateBindlessDescriptors(FVulkanCommandListContext* CmdListContext, FVulkanCommandBuffer& CommandBuffer);
+	void UpdateBindlessDescriptors(FVulkanCommandListContext& Context);
 
 	inline void BindDescriptorSets(VkCommandBuffer CmdBuffer)
 	{
@@ -284,7 +286,7 @@ protected:
 	FVulkanRHIGraphicsPipelineState* GfxPipeline;
 
 	template<bool bUseDynamicGlobalUBs>
-	bool InternalUpdateDescriptorSets(FVulkanCommandListContext* CmdListContext, FVulkanCommandBuffer* CmdBuffer);
+	bool InternalUpdateDescriptorSets(FVulkanCommandListContext& Context);
 
 	friend class FVulkanPendingGfxState;
 	friend class FVulkanCommandListContext;
diff --git a/Engine/Source/Runtime/VulkanRHI/Private/VulkanRayTracing.cpp b/Engine/Source/Runtime/VulkanRHI/Private/VulkanRayTracing.cpp
--- a/Engine/Source/Runtime/VulkanRHI/Private/VulkanRayTracing.cpp
+++ b/Engine/Source/Runtime/VulkanRHI/Private/VulkanRayTracing.cpp
@@ -1953,7 +1953,7 @@ void FVulkanRayTracingCompactionRequestHandler::Update(FVulkanCommandListContext
 }
 
 
-static void SetRayGenResources(FVulkanDevice& Device, FVulkanCommandBuffer& CmdBuffer, const FRayTracingShaderBindings& InGlobalResourceBindings, FVulkanRayTracingShaderTable* ShaderTable)
+static void SetRayGenResources(FVulkanDevice& Device, FVulkanCommandListContext& Context, const FRayTracingShaderBindings& InGlobalResourceBindings, FVulkanRayTracingShaderTable* ShaderTable)
 {
 	TArray<const FVulkanUniformBuffer*> UniformBuffers;
 	UniformBuffers.Reserve(UE_ARRAY_COUNT(InGlobalResourceBindings.UniformBuffers));
@@ -1994,7 +1994,7 @@ static void SetRayGenResources(FVulkanDevice& Device, FVulkanCommandBuffer& CmdB
 				NumSkippedSlots++;
 			}
 		}
-		Device.GetBindlessDescriptorManager()->RegisterUniformBuffers(CmdBuffer, VK_PIPELINE_BIND_POINT_RAY_TRACING_KHR, StageUBs);
+		Device.GetBindlessDescriptorManager()->RegisterUniformBuffers(Context, VK_PIPELINE_BIND_POINT_RAY_TRACING_KHR, StageUBs);
 	}
 
 	// Add all the UBs references by the shader table
@@ -2026,7 +2026,7 @@ void FVulkanCommandListContext::RHIRayTraceDispatch(
 	checkSlow(InRayGenShader->GetFrequency() == SF_RayGen);
 	const VkStridedDeviceAddressRegionKHR RayGenRegion = ShaderTable->CommitRayGenShader(*this, Pipeline->GetShaderIndex(RayGenShader), Pipeline->GetShaderHandles(SF_RayGen));
 
-	SetRayGenResources(Device, CommandBuffer, InGlobalResourceBindings, ShaderTable);
+	SetRayGenResources(Device, *this, InGlobalResourceBindings, ShaderTable);
 
 	VulkanRHI::vkCmdTraceRaysKHR(
 		CommandBuffer.GetHandle(),
@@ -2057,7 +2057,7 @@ void FVulkanCommandListContext::RHIRayTraceDispatchIndirect(
 	checkSlow(InRayGenShader->GetFrequency() == SF_RayGen);
 	const VkStridedDeviceAddressRegionKHR RayGenRegion = ShaderTable->CommitRayGenShader(*this, Pipeline->GetShaderIndex(RayGenShader), Pipeline->GetShaderHandles(SF_RayGen));
 
-	SetRayGenResources(Device, CommandBuffer, InGlobalResourceBindings, ShaderTable);
+	SetRayGenResources(Device, *this, InGlobalResourceBindings, ShaderTable);
 
 	FVulkanBuffer* ArgumentBuffer = ResourceCast(InArgumentBuffer);
 	const VkDeviceAddress IndirectDeviceAddress = ArgumentBuffer->GetDeviceAddress() + InArgumentOffset;
