VulkanRHI: Clean up:  associate CommandBufferPool to Queue instead of Context.

#rnx

diff --git a/Engine/Source/Runtime/VulkanRHI/Private/VulkanBuffer.cpp b/Engine/Source/Runtime/VulkanRHI/Private/VulkanBuffer.cpp
--- a/Engine/Source/Runtime/VulkanRHI/Private/VulkanBuffer.cpp
+++ b/Engine/Source/Runtime/VulkanRHI/Private/VulkanBuffer.cpp
@@ -360,9 +360,16 @@ void FVulkanBuffer::Unlock(FRHICommandListBase& RHICmdList)
 				// We need to do this on the active command buffer instead of using an upload command buffer. The high level code sometimes reuses the same
 				// buffer in sequences of upload / dispatch, upload / dispatch, so we need to order the copy commands correctly with respect to the dispatches.
 				// Unless this is the first time any data is pushed into this buffer, then ordering doesn't matter and UploadContext can be used
-				FVulkanCommandBuffer& CommandBuffer = PendingLock.FirstLock ? 
-					FVulkanUploadContext::Get(CmdList).GetCommandBuffer() :
-					FVulkanCommandListContext::Get(CmdList).GetCommandBuffer();
+				FVulkanContextCommon* CommonContext = nullptr;
+				if (PendingLock.FirstLock)
+				{
+					CommonContext = &FVulkanUploadContext::Get(CmdList);
+				}
+				else
+				{
+					CommonContext = &FVulkanCommandListContext::Get(CmdList);
+				}
+				FVulkanCommandBuffer& CommandBuffer = CommonContext->GetCommandBuffer();
 				check(CommandBuffer.IsOutsideRenderPass());
 				VkCommandBuffer CommandBufferHandle = CommandBuffer.GetHandle();
 
@@ -379,7 +386,7 @@ void FVulkanBuffer::Unlock(FRHICommandListBase& RHICmdList)
 				VkMemoryBarrier BarrierAfter = { VK_STRUCTURE_TYPE_MEMORY_BARRIER, nullptr, VK_ACCESS_TRANSFER_WRITE_BIT, VK_ACCESS_MEMORY_READ_BIT | VK_ACCESS_MEMORY_WRITE_BIT };
 				VulkanRHI::vkCmdPipelineBarrier(CommandBufferHandle, VK_PIPELINE_STAGE_TRANSFER_BIT, VK_PIPELINE_STAGE_ALL_COMMANDS_BIT, 0, 1, &BarrierAfter, 0, nullptr, 0, nullptr);
 
-				Buffer->GetParent()->GetStagingManager().ReleaseBuffer(&CommandBuffer.GetOwner()->GetContext(), StagingBuffer);
+				Buffer->GetParent()->GetStagingManager().ReleaseBuffer(CommonContext, StagingBuffer);
 			}
 		});
 	}
diff --git a/Engine/Source/Runtime/VulkanRHI/Private/VulkanCommandBuffer.cpp b/Engine/Source/Runtime/VulkanRHI/Private/VulkanCommandBuffer.cpp
--- a/Engine/Source/Runtime/VulkanRHI/Private/VulkanCommandBuffer.cpp
+++ b/Engine/Source/Runtime/VulkanRHI/Private/VulkanCommandBuffer.cpp
@@ -18,7 +18,7 @@ FVulkanCommandBuffer::FVulkanCommandBuffer(FVulkanDevice& InDevice, FVulkanComma
 	: Device(InDevice)
 	, CommandBufferPool(InCommandBufferPool)
 #if RHI_NEW_GPU_PROFILER
-	, EventStream(InCommandBufferPool.GetContext().Queue.GetProfilerQueue())
+	, EventStream(InCommandBufferPool.GetQueue().GetProfilerQueue())
 #endif
 {
 	FScopeLock ScopeLock(CommandBufferPool.GetCS());
@@ -165,8 +165,7 @@ void FVulkanCommandBuffer::Begin(FVulkanQueryPool* QueryPool)
 	if (Device.SupportsBindless())
 	{
 		FVulkanBindlessDescriptorManager* BindlessDescriptorManager = Device.GetBindlessDescriptorManager();
-		FVulkanQueue& Queue = GetOwner()->GetContext().Queue;
-		const VkPipelineStageFlags SupportedStages = Queue.GetSupportedStageBits();
+		const VkPipelineStageFlags SupportedStages = CommandBufferPool.GetQueue().GetSupportedStageBits();
 		BindlessDescriptorManager->BindDescriptorBuffers(CommandBufferHandle, SupportedStages);
 	}
 
@@ -278,11 +277,15 @@ void FVulkanCommandBuffer::EndSplitBarrier(VkEvent BarrierEvent, const VkDepende
 }
 
 
-FVulkanCommandBufferPool::FVulkanCommandBufferPool(FVulkanDevice& InDevice, FVulkanContextCommon& InContext)
+FVulkanCommandBufferPool::FVulkanCommandBufferPool(FVulkanDevice& InDevice, FVulkanQueue& InQueue)
 	: Device(InDevice)
-	, Context(InContext)
-	, Handle(VK_NULL_HANDLE)
+	, Queue(InQueue)
 {
+	VkCommandPoolCreateInfo CmdPoolInfo;
+	ZeroVulkanStruct(CmdPoolInfo, VK_STRUCTURE_TYPE_COMMAND_POOL_CREATE_INFO);
+	CmdPoolInfo.queueFamilyIndex = InQueue.GetFamilyIndex();
+	CmdPoolInfo.flags = VK_COMMAND_POOL_CREATE_RESET_COMMAND_BUFFER_BIT; // :todo: Investigate use of VK_COMMAND_POOL_CREATE_TRANSIENT_BIT?
+	VERIFYVULKANRESULT(VulkanRHI::vkCreateCommandPool(Device.GetInstanceHandle(), &CmdPoolInfo, VULKAN_CPU_ALLOCATOR, &Handle));
 }
 
 FVulkanCommandBufferPool::~FVulkanCommandBufferPool()
@@ -303,17 +306,6 @@ FVulkanCommandBufferPool::~FVulkanCommandBufferPool()
 	Handle = VK_NULL_HANDLE;
 }
 
-void FVulkanCommandBufferPool::Create(uint32 QueueFamilyIndex)
-{
-	VkCommandPoolCreateInfo CmdPoolInfo;
-	ZeroVulkanStruct(CmdPoolInfo, VK_STRUCTURE_TYPE_COMMAND_POOL_CREATE_INFO);
-	CmdPoolInfo.queueFamilyIndex =  QueueFamilyIndex;
-	//#todo-rco: Should we use VK_COMMAND_POOL_CREATE_TRANSIENT_BIT?
-	CmdPoolInfo.flags = VK_COMMAND_POOL_CREATE_RESET_COMMAND_BUFFER_BIT;
-	VERIFYVULKANRESULT(VulkanRHI::vkCreateCommandPool(Device.GetInstanceHandle(), &CmdPoolInfo, VULKAN_CPU_ALLOCATOR, &Handle));
-}
-
-
 void FVulkanCommandBufferPool::FreeUnusedCmdBuffers(FVulkanQueue* InQueue, bool bTrimMemory)
 {
 #if VULKAN_DELETE_STALE_CMDBUFFERS
diff --git a/Engine/Source/Runtime/VulkanRHI/Private/VulkanCommandBuffer.h b/Engine/Source/Runtime/VulkanRHI/Private/VulkanCommandBuffer.h
--- a/Engine/Source/Runtime/VulkanRHI/Private/VulkanCommandBuffer.h
+++ b/Engine/Source/Runtime/VulkanRHI/Private/VulkanCommandBuffer.h
@@ -34,11 +34,6 @@ protected:
 	~FVulkanCommandBuffer();
 
 public:
-	FVulkanCommandBufferPool* GetOwner()
-	{
-		return &CommandBufferPool;
-	}
-
 	bool IsInsideRenderPass() const
 	{
 		return State == EState::IsInsideRenderPass;
@@ -189,7 +184,7 @@ public:
 class FVulkanCommandBufferPool
 {
 public:
-	FVulkanCommandBufferPool(FVulkanDevice& InDevice, FVulkanContextCommon& InContext);
+	FVulkanCommandBufferPool(FVulkanDevice& InDevice, FVulkanQueue& InQueue);
 	~FVulkanCommandBufferPool();
 
 	VkCommandPool GetHandle() const
@@ -202,18 +197,18 @@ public:
 		return &CS;
 	}
 
-	void FreeUnusedCmdBuffers(FVulkanQueue* Queue, bool bTrimMemory);
-
-	FVulkanContextCommon& GetContext()
+	FVulkanQueue& GetQueue()
 	{
-		return Context;
+		return Queue;
 	}
 
+	void FreeUnusedCmdBuffers(FVulkanQueue* Queue, bool bTrimMemory);
+
 private:
 	FVulkanDevice& Device;
-	FVulkanContextCommon& Context;
+	FVulkanQueue& Queue;
 
-	VkCommandPool Handle;
+	VkCommandPool Handle = VK_NULL_HANDLE;
 
 	TArray<FVulkanCommandBuffer*> CmdBuffers;
 	TArray<FVulkanCommandBuffer*> FreeCmdBuffers;
@@ -222,7 +217,5 @@ private:
 
 	FVulkanCommandBuffer* Create();
 
-	void Create(uint32 QueueFamilyIndex);
-
 	friend FVulkanContextCommon;
 };
diff --git a/Engine/Source/Runtime/VulkanRHI/Private/VulkanContext.cpp b/Engine/Source/Runtime/VulkanRHI/Private/VulkanContext.cpp
--- a/Engine/Source/Runtime/VulkanRHI/Private/VulkanContext.cpp
+++ b/Engine/Source/Runtime/VulkanRHI/Private/VulkanContext.cpp
@@ -15,9 +15,8 @@
 FVulkanContextCommon::FVulkanContextCommon(FVulkanDevice& InDevice, FVulkanQueue& InQueue)
 	: Device(InDevice)
 	, Queue(InQueue)
-	, Pool(InDevice, *this)
+	, Pool(InDevice, InQueue)
 {
-	Pool.Create(Queue.GetFamilyIndex());
 }
 
 FVulkanContextCommon::~FVulkanContextCommon()
