VulkanRHI: Fix debug build

#rnx

diff --git a/Engine/Source/Runtime/VulkanRHI/Private/VulkanPendingState.cpp b/Engine/Source/Runtime/VulkanRHI/Private/VulkanPendingState.cpp
--- a/Engine/Source/Runtime/VulkanRHI/Private/VulkanPendingState.cpp
+++ b/Engine/Source/Runtime/VulkanRHI/Private/VulkanPendingState.cpp
@@ -462,7 +462,7 @@ void FVulkanPendingGfxState::PrepareForDraw(FVulkanCommandListContext& Context)
 					if (VertexInputStateInfo.Attributes[AttributeIndex].binding == CurrBinding.binding)
 					{
 						uint64 VertexShaderKey = GetCurrentShaderKey(ShaderStage::Vertex);
-						FVulkanVertexShader* VertexShader = Device->GetShaderFactory().LookupShader<FVulkanVertexShader>(VertexShaderKey);
+						FVulkanVertexShader* VertexShader = Device.GetShaderFactory().LookupShader<FVulkanVertexShader>(VertexShaderKey);
 						UE_LOG(LogVulkanRHI, Warning, TEXT("Missing input assembly binding on location %d in Vertex shader '%s'"),
 							CurrBinding.binding,
 							VertexShader ? *VertexShader->GetDebugName() : TEXT("Null"));
