VulkanRHI: Clean up:  Remove RHI pointer from Contexts

#rnx

diff --git a/Engine/Source/Runtime/VulkanRHI/Private/VulkanContext.cpp b/Engine/Source/Runtime/VulkanRHI/Private/VulkanContext.cpp
--- a/Engine/Source/Runtime/VulkanRHI/Private/VulkanContext.cpp
+++ b/Engine/Source/Runtime/VulkanRHI/Private/VulkanContext.cpp
@@ -173,9 +173,8 @@ void FVulkanContextCommon::PrepareNewCommandBuffer(FVulkanPayload& Payload)
 }
 
 
-FVulkanCommandListContext::FVulkanCommandListContext(FVulkanDynamicRHI* InRHI, FVulkanDevice& InDevice, ERHIPipeline InPipeline, FVulkanCommandListContext* InImmediate)
+FVulkanCommandListContext::FVulkanCommandListContext(FVulkanDevice& InDevice, ERHIPipeline InPipeline, FVulkanCommandListContext* InImmediate)
 	: FVulkanContextCommon(InDevice, *InDevice.GetQueue(InPipeline))
-	, RHI(InRHI)
 	, Immediate(InImmediate)
 	, RHIPipeline(InPipeline)
 #if (RHI_NEW_GPU_PROFILER == 0)
@@ -214,8 +213,7 @@ FVulkanCommandListContext::~FVulkanCommandListContext()
 #endif
 	}
 
-	delete PendingGfxState;
-	delete PendingComputeState;
+	ReleasePendingState();
 }
 
 
diff --git a/Engine/Source/Runtime/VulkanRHI/Private/VulkanContext.h b/Engine/Source/Runtime/VulkanRHI/Private/VulkanContext.h
--- a/Engine/Source/Runtime/VulkanRHI/Private/VulkanContext.h
+++ b/Engine/Source/Runtime/VulkanRHI/Private/VulkanContext.h
@@ -129,11 +129,6 @@ public:
 	// Should only ne used when we are certain there are no other pending contexts (like UploadContext)
 	void FlushCommands(EVulkanFlushFlags FlushFlags = EVulkanFlushFlags::None);
 
-	uint64 GetFrameCounter() const
-	{
-		return FrameCounter;
-	}
-
 	// Public immutable references, access them directly
 	FVulkanDevice& Device;
 	FVulkanQueue& Queue;
@@ -190,10 +185,6 @@ public:
 		return GetCurrentTimestampQueryPool(GetPayload(EPhase::Execute));
 	}
 
-protected:
-	// Number of times EndFrame() has been called on this context
-	uint64 FrameCounter = 0;
-
 private:
 	FVulkanQueryPool* GetCurrentTimestampQueryPool(FVulkanPayload& Payload);
 	void PrepareNewCommandBuffer(FVulkanPayload& Payload);
@@ -217,7 +208,7 @@ private:
 class FVulkanCommandListContext : public FVulkanContextCommon, public IRHICommandContext
 {
 public:
-	FVulkanCommandListContext(FVulkanDynamicRHI* InRHI, FVulkanDevice& InDevice, ERHIPipeline InPipeline, FVulkanCommandListContext* InImmediate);
+	FVulkanCommandListContext(FVulkanDevice& InDevice, ERHIPipeline InPipeline, FVulkanCommandListContext* InImmediate);
 	virtual ~FVulkanCommandListContext();
 
 	static inline FVulkanCommandListContext& Get(FRHICommandListBase& RHICmdList)
@@ -382,7 +373,6 @@ public:
 	FVulkanQueryPool* GetCurrentOcclusionQueryPool();
 
 protected:
-	FVulkanDynamicRHI* RHI;
 	FVulkanCommandListContext* Immediate;
 	const ERHIPipeline RHIPipeline;
 
@@ -443,7 +433,7 @@ public:
 		return static_cast<FVulkanCommandListContextImmediate&>(RHICmdList.GetContext().GetLowestLevelContext());
 	}
 
-	FVulkanCommandListContextImmediate(FVulkanDynamicRHI* InRHI, FVulkanDevice& InDevice);
+	FVulkanCommandListContextImmediate(FVulkanDevice& InDevice);
 };
 
 class FVulkanUploadContext : public FVulkanContextCommon, public IRHIUploadContext
diff --git a/Engine/Source/Runtime/VulkanRHI/Private/VulkanDevice.cpp b/Engine/Source/Runtime/VulkanRHI/Private/VulkanDevice.cpp
--- a/Engine/Source/Runtime/VulkanRHI/Private/VulkanDevice.cpp
+++ b/Engine/Source/Runtime/VulkanRHI/Private/VulkanDevice.cpp
@@ -1419,7 +1419,7 @@ void FVulkanDevice::InitGPU()
 	// always look in the saved directory (for the cache from previous run that wasn't moved over to stage directory)
 	CacheFilenames.Add(VulkanRHI::GetPipelineCacheFilename());
 
-	ImmediateContext = new FVulkanCommandListContextImmediate(GVulkanRHI, *this);
+	ImmediateContext = new FVulkanCommandListContextImmediate(*this);
 
 	extern TAutoConsoleVariable<int32> GRHIThreadCvar;
 	if (GRHIThreadCvar->GetInt() > 1)
@@ -1427,7 +1427,7 @@ void FVulkanDevice::InitGPU()
 		int32 Num = FTaskGraphInterface::Get().GetNumWorkerThreads();
 		for (int32 Index = 0; Index < Num; Index++)
 		{
-			FVulkanCommandListContext* CmdContext = new FVulkanCommandListContext(GVulkanRHI, *this, ERHIPipeline::Graphics, ImmediateContext);
+			FVulkanCommandListContext* CmdContext = new FVulkanCommandListContext(*this, ERHIPipeline::Graphics, ImmediateContext);
 			CommandContexts[ERHIPipeline::Graphics].Add(CmdContext);
 		}
 	}
@@ -1626,7 +1626,7 @@ FVulkanCommandListContext* FVulkanDevice::AcquireDeferredContext(ERHIPipeline Pi
 
 	if (Contexts.Num() == 0)
 	{
-		return new FVulkanCommandListContext(GVulkanRHI, *this, Pipeline, ImmediateContext);
+		return new FVulkanCommandListContext(*this, Pipeline, ImmediateContext);
 	}
 
 	return Contexts.Pop(EAllowShrinking::No);
diff --git a/Engine/Source/Runtime/VulkanRHI/Private/VulkanRHI.cpp b/Engine/Source/Runtime/VulkanRHI/Private/VulkanRHI.cpp
--- a/Engine/Source/Runtime/VulkanRHI/Private/VulkanRHI.cpp
+++ b/Engine/Source/Runtime/VulkanRHI/Private/VulkanRHI.cpp
@@ -421,8 +421,8 @@ FDynamicRHI* FVulkanDynamicRHIModule::CreateRHI(ERHIFeatureLevel::Type InRequest
 IMPLEMENT_MODULE(FVulkanDynamicRHIModule, VulkanRHI);
 
 
-FVulkanCommandListContextImmediate::FVulkanCommandListContextImmediate(FVulkanDynamicRHI* InRHI, FVulkanDevice& InDevice)
-	: FVulkanCommandListContext(InRHI, InDevice, ERHIPipeline::Graphics, nullptr)
+FVulkanCommandListContextImmediate::FVulkanCommandListContextImmediate(FVulkanDevice& InDevice)
+	: FVulkanCommandListContext(InDevice, ERHIPipeline::Graphics, nullptr)
 {
 }
 
@@ -1115,8 +1115,6 @@ void FVulkanDynamicRHI::RHIEndFrame_RenderThread(FRHICommandListImmediate& RHICm
 
 		Context.Device.GetBindlessDescriptorManager()->UpdateUBAllocator();
 		Context.Device.GetTempBlockAllocator().UpdateBlocks();
-
-		++Context.FrameCounter;
 	});
 
 	FDynamicRHI::RHIEndFrame_RenderThread(RHICmdList);
@@ -1214,7 +1212,7 @@ void FVulkanCommandListContext::RHIBeginDrawingViewport(FRHIViewport* ViewportRH
 	//FRCLog::Printf(FString::Printf(TEXT("FVulkanCommandListContext::RHIBeginDrawingViewport\n")));
 	check(ViewportRHI);
 	FVulkanViewport* Viewport = ResourceCast(ViewportRHI);
-	RHI->DrawingViewport = Viewport;
+	FVulkanDynamicRHI::Get().DrawingViewport = Viewport;
 
 	FRHICustomPresent* CustomPresent = Viewport->GetCustomPresent();
 	if (CustomPresent)
@@ -1229,7 +1227,7 @@ void FVulkanCommandListContext::RHIEndDrawingViewport(FRHIViewport* ViewportRHI,
 	//FRCLog::Printf(FString::Printf(TEXT("FVulkanCommandListContext::RHIEndDrawingViewport()")));
 	check(IsImmediate());
 	FVulkanViewport* Viewport = ResourceCast(ViewportRHI);
-	check(Viewport == RHI->DrawingViewport);
+	check(Viewport == FVulkanDynamicRHI::Get().DrawingViewport);
 
 	//#todo-rco: Unbind all pending state
 /*
@@ -1245,7 +1243,7 @@ void FVulkanCommandListContext::RHIEndDrawingViewport(FRHIViewport* ViewportRHI,
 		//#todo-rco: Check for r.FinishCurrentFrame
 	}
 
-	RHI->DrawingViewport = nullptr;
+	FVulkanDynamicRHI::Get().DrawingViewport = nullptr;
 }
 
 #if WITH_RHI_BREADCRUMBS
diff --git a/Engine/Source/Runtime/VulkanRHI/Private/VulkanRHIPrivate.h b/Engine/Source/Runtime/VulkanRHI/Private/VulkanRHIPrivate.h
--- a/Engine/Source/Runtime/VulkanRHI/Private/VulkanRHIPrivate.h
+++ b/Engine/Source/Runtime/VulkanRHI/Private/VulkanRHIPrivate.h
@@ -243,7 +243,6 @@ protected:
 	}
 
 	bool bCalculatedHash = false;
-	void CalculateRenderPassHashes(const FRHISetRenderTargetsInfo& RTInfo);
 
 	friend class FVulkanPipelineStateCacheManager;
 	friend struct FGfxPipelineDesc;
diff --git a/Engine/Source/Runtime/VulkanRHI/Private/VulkanRenderTarget.cpp b/Engine/Source/Runtime/VulkanRHI/Private/VulkanRenderTarget.cpp
--- a/Engine/Source/Runtime/VulkanRHI/Private/VulkanRenderTarget.cpp
+++ b/Engine/Source/Runtime/VulkanRHI/Private/VulkanRenderTarget.cpp
@@ -466,7 +466,7 @@ void FVulkanDynamicRHI::RHIRead3DSurfaceFloatData(FRHITexture* TextureRHI, FIntR
 
 VkFormat FVulkanCommandListContext::GetSwapchainImageFormat() const
 {
-	TArray<FVulkanViewport*>& viewports = RHI->GetViewports();
+	TArray<FVulkanViewport*>& viewports = FVulkanDynamicRHI::Get().GetViewports();
 	if (viewports.Num() == 0)
 	{
 		return VK_FORMAT_UNDEFINED;
@@ -477,7 +477,7 @@ VkFormat FVulkanCommandListContext::GetSwapchainImageFormat() const
 
 FVulkanSwapChain* FVulkanCommandListContext::GetSwapChain() const
 {
-	TArray<FVulkanViewport*>& viewports = RHI->GetViewports();
+	TArray<FVulkanViewport*>& viewports = FVulkanDynamicRHI::Get().GetViewports();
 	uint32 numViewports = viewports.Num();
 
 	if (viewports.Num() == 0)
@@ -490,7 +490,7 @@ FVulkanSwapChain* FVulkanCommandListContext::GetSwapChain() const
 
 bool FVulkanCommandListContext::IsSwapchainImage(FRHITexture* InTexture) const
 {
-	TArray<FVulkanViewport*>& Viewports = RHI->GetViewports();
+	TArray<FVulkanViewport*>& Viewports = FVulkanDynamicRHI::Get().GetViewports();
 	uint32 NumViewports = Viewports.Num();
 
 	for (uint32 i = 0; i < NumViewports; i++)
