VulkanRHI: Use TConditionalScopeLock instead of FPSOOptionalLock.

#rnx

diff --git a/Engine/Source/Runtime/VulkanRHI/Private/VulkanPipeline.cpp b/Engine/Source/Runtime/VulkanRHI/Private/VulkanPipeline.cpp
--- a/Engine/Source/Runtime/VulkanRHI/Private/VulkanPipeline.cpp
+++ b/Engine/Source/Runtime/VulkanRHI/Private/VulkanPipeline.cpp
@@ -2073,26 +2073,6 @@ void FVulkanPipelineStateCacheManager::NotifyDeletedGraphicsPSO(FRHIGraphicsPipe
 }
 
 
-struct FPSOOptionalLock
-{
-	FCriticalSection* CriticalSection;
-	FPSOOptionalLock(FCriticalSection* InSynchObject)
-	{
-		CriticalSection = InSynchObject;
-		if (CriticalSection)
-		{
-			CriticalSection->Lock();
-		}
-	}
-	~FPSOOptionalLock()
-	{
-		if (CriticalSection)
-		{
-			CriticalSection->Unlock();
-		}
-	}
-};
-
 static FCriticalSection CreateGraphicsPSOMutex;
 
 FGraphicsPipelineStateRHIRef FVulkanPipelineStateCacheManager::RHICreateGraphicsPipelineState(const FGraphicsPipelineStateInitializer& Initializer)
@@ -2109,7 +2089,7 @@ FGraphicsPipelineStateRHIRef FVulkanPipelineStateCacheManager::RHICreateGraphics
 		|| (ThreadingMode == ESingleThreadedPSOCreateMode::Precompile && bIsPrecache)
 		|| (ThreadingMode == ESingleThreadedPSOCreateMode::NonPrecompiled && !bIsPrecache);
 
-	FPSOOptionalLock PSOSingleThreadedLock(bShouldLock ? &CreateGraphicsPSOMutex : nullptr);
+	UE::TConditionalScopeLock PSOSingleThreadedLock(CreateGraphicsPSOMutex, bShouldLock);
 
 	FVulkanPSOKey Key;
 	FGfxPipelineDesc Desc;
