VulkanRHI: Clean up: Remove duplicate parameters and prefer passing Context over CommandBuffer

#rnx

diff --git a/Engine/Source/Runtime/VulkanRHI/Private/VulkanRenderTarget.cpp b/Engine/Source/Runtime/VulkanRHI/Private/VulkanRenderTarget.cpp
--- a/Engine/Source/Runtime/VulkanRHI/Private/VulkanRenderTarget.cpp
+++ b/Engine/Source/Runtime/VulkanRHI/Private/VulkanRenderTarget.cpp
@@ -569,7 +569,7 @@ void FVulkanCommandListContext::RHIBeginRenderPass(const FRHIRenderPassInfo& InI
 
 	FVulkanFramebuffer* Framebuffer = Device.GetRenderPassManager().GetOrCreateFramebuffer(RTInfo, RTLayout, RenderPass);
 	checkf(RenderPass != nullptr && Framebuffer != nullptr, TEXT("RenderPass not started! Bad combination of values? Depth %p #Color %d Color0 %p"), (void*)InInfo.DepthStencilRenderTarget.DepthStencilTarget, InInfo.GetNumColorRenderTargets(), (void*)InInfo.ColorRenderTargets[0].RenderTarget);
-	Device.GetRenderPassManager().BeginRenderPass(*this, Device, &GetCommandBuffer(), InInfo, RTLayout, RenderPass, Framebuffer);
+	Device.GetRenderPassManager().BeginRenderPass(*this, Device, InInfo, RTLayout, RenderPass, Framebuffer);
 
 	check(!CurrentRenderPass);
 	CurrentRenderPass = RenderPass;
@@ -578,8 +578,7 @@ void FVulkanCommandListContext::RHIBeginRenderPass(const FRHIRenderPassInfo& InI
 
 void FVulkanCommandListContext::RHIEndRenderPass()
 {
-	FVulkanCommandBuffer& CommandBuffer = GetCommandBuffer();
-	Device.GetRenderPassManager().EndRenderPass(&CommandBuffer);
+	Device.GetRenderPassManager().EndRenderPass(*this);
 
 	const bool bHasOcclusionQueries = (RenderPassInfo.NumOcclusionQueries > 0);
 	if (bHasOcclusionQueries)
diff --git a/Engine/Source/Runtime/VulkanRHI/Private/VulkanRenderpass.cpp b/Engine/Source/Runtime/VulkanRHI/Private/VulkanRenderpass.cpp
--- a/Engine/Source/Runtime/VulkanRHI/Private/VulkanRenderpass.cpp
+++ b/Engine/Source/Runtime/VulkanRHI/Private/VulkanRenderpass.cpp
@@ -126,7 +126,7 @@ FVulkanFramebuffer* FVulkanRenderPassManager::GetOrCreateFramebuffer(const FRHIS
 	return Framebuffer;
 }
 
-void FVulkanRenderPassManager::BeginRenderPass(FVulkanCommandListContext& Context, FVulkanDevice& InDevice, FVulkanCommandBuffer* CmdBuffer, const FRHIRenderPassInfo& RPInfo, const FVulkanRenderTargetLayout& RTLayout, FVulkanRenderPass* RenderPass, FVulkanFramebuffer* Framebuffer)
+void FVulkanRenderPassManager::BeginRenderPass(FVulkanCommandListContext& Context, FVulkanDevice& InDevice, const FRHIRenderPassInfo& RPInfo, const FVulkanRenderTargetLayout& RTLayout, FVulkanRenderPass* RenderPass, FVulkanFramebuffer* Framebuffer)
 {
 	// (NumRT + 1 [Depth] ) * 2 [surface + resolve]
 	VkClearValue ClearValues[(MaxSimultaneousRenderTargets + 1) * 2];
@@ -204,9 +204,11 @@ void FVulkanRenderPassManager::BeginRenderPass(FVulkanCommandListContext& Contex
 
 	ensure(ClearValueIndex <= RenderPass->GetNumUsedClearValues());
 
-	Barrier.Execute(CmdBuffer);
+	FVulkanCommandBuffer& CommandBuffer = Context.GetCommandBuffer();
 
-	CmdBuffer->BeginRenderPass(RenderPass->GetLayout(), RenderPass, Framebuffer, ClearValues);
+	Barrier.Execute(&CommandBuffer);
+
+	CommandBuffer.BeginRenderPass(RenderPass->GetLayout(), RenderPass, Framebuffer, ClearValues);
 
 	{
 		const VkExtent3D& Extents = RTLayout.GetExtent3D();
@@ -214,11 +216,12 @@ void FVulkanRenderPassManager::BeginRenderPass(FVulkanCommandListContext& Contex
 	}
 }
 
-void FVulkanRenderPassManager::EndRenderPass(FVulkanCommandBuffer* CmdBuffer)
+void FVulkanRenderPassManager::EndRenderPass(FVulkanCommandListContext& Context)
 {
-	CmdBuffer->EndRenderPass();
+	FVulkanCommandBuffer& CommandBuffer = Context.GetCommandBuffer();
+	CommandBuffer.EndRenderPass();
 
-	VulkanRHI::DebugHeavyWeightBarrier(CmdBuffer->GetHandle(), 1);
+	VulkanRHI::DebugHeavyWeightBarrier(CommandBuffer.GetHandle(), 1);
 }
 
 void FVulkanRenderPassManager::NotifyDeletedRenderTarget(VkImage Image)
diff --git a/Engine/Source/Runtime/VulkanRHI/Private/VulkanRenderpass.h b/Engine/Source/Runtime/VulkanRHI/Private/VulkanRenderpass.h
--- a/Engine/Source/Runtime/VulkanRHI/Private/VulkanRenderpass.h
+++ b/Engine/Source/Runtime/VulkanRHI/Private/VulkanRenderpass.h
@@ -888,8 +888,8 @@ public:
 		return RenderPass;
 	}
 
-	void BeginRenderPass(FVulkanCommandListContext& Context, FVulkanDevice& InDevice, FVulkanCommandBuffer* CmdBuffer, const FRHIRenderPassInfo& RPInfo, const FVulkanRenderTargetLayout& RTLayout, FVulkanRenderPass* RenderPass, FVulkanFramebuffer* Framebuffer);
-	void EndRenderPass(FVulkanCommandBuffer* CmdBuffer);
+	void BeginRenderPass(FVulkanCommandListContext& Context, FVulkanDevice& InDevice, const FRHIRenderPassInfo& RPInfo, const FVulkanRenderTargetLayout& RTLayout, FVulkanRenderPass* RenderPass, FVulkanFramebuffer* Framebuffer);
+	void EndRenderPass(FVulkanCommandListContext& Context);
 
 	FRWLock RenderPassesLock;
 	FRWLock FramebuffersLock;
