VulkanRHI: Clean up: Remove redundant param.  Bundle RenderPass and FrameBuffer in a struct for BeginRenderPass.  Remove VulkanRHI::FDeviceChild from FVulkanRenderPassManager.

#rnx

diff --git a/Engine/Source/Runtime/VulkanRHI/Private/VulkanCommandBuffer.cpp b/Engine/Source/Runtime/VulkanRHI/Private/VulkanCommandBuffer.cpp
--- a/Engine/Source/Runtime/VulkanRHI/Private/VulkanCommandBuffer.cpp
+++ b/Engine/Source/Runtime/VulkanRHI/Private/VulkanCommandBuffer.cpp
@@ -9,6 +9,7 @@
 #include "VulkanDescriptorSets.h"
 #include "VulkanMemory.h"
 #include "VulkanRayTracing.h"
+#include "VulkanRenderpass.h"
 
 
 #define CMD_BUFFER_TIME_TO_WAIT_BEFORE_DELETING		10
@@ -77,15 +78,17 @@ void FVulkanCommandBuffer::EndRenderPass()
 	State = EState::IsInsideBegin;
 }
 
-void FVulkanCommandBuffer::BeginRenderPass(const FVulkanRenderTargetLayout& Layout, FVulkanRenderPass* RenderPass, FVulkanFramebuffer* Framebuffer, const VkClearValue* AttachmentClearValues)
+void FVulkanCommandBuffer::BeginRenderPass(const FVulkanBeginRenderPassInfo& BeginRenderPassInfo, const VkClearValue* AttachmentClearValues)
 {
 	checkf(IsOutsideRenderPass(), TEXT("Can't BeginRP as already inside one! CmdBuffer 0x%p State=%d"), CommandBufferHandle, (int32)State);
 
+	const FVulkanRenderTargetLayout& Layout = BeginRenderPassInfo.RenderPass.GetLayout();
+
 	VkRenderPassBeginInfo Info;
 	ZeroVulkanStruct(Info, VK_STRUCTURE_TYPE_RENDER_PASS_BEGIN_INFO);
-	Info.renderPass = RenderPass->GetHandle();
-	Info.framebuffer = Framebuffer->GetHandle();
-	Info.renderArea = Framebuffer->GetRenderArea();
+	Info.renderPass = BeginRenderPassInfo.RenderPass.GetHandle();
+	Info.framebuffer = BeginRenderPassInfo.Framebuffer.GetHandle();
+	Info.renderArea = BeginRenderPassInfo.Framebuffer.GetRenderArea();
 	Info.clearValueCount = Layout.GetNumUsedClearValues();
 	Info.pClearValues = AttachmentClearValues;
 
diff --git a/Engine/Source/Runtime/VulkanRHI/Private/VulkanCommandBuffer.h b/Engine/Source/Runtime/VulkanRHI/Private/VulkanCommandBuffer.h
--- a/Engine/Source/Runtime/VulkanRHI/Private/VulkanCommandBuffer.h
+++ b/Engine/Source/Runtime/VulkanRHI/Private/VulkanCommandBuffer.h
@@ -17,6 +17,7 @@ class FVulkanRenderTargetLayout;
 class FVulkanQueue;
 class FVulkanDescriptorPoolSetContainer;
 class FVulkanQueryPool;
+struct FVulkanBeginRenderPassInfo;
 
 namespace VulkanRHI
 {
@@ -96,7 +97,7 @@ public:
 	uint8 bHasStencilRef				: 1 = 0;
 
 	// You never want to call Begin/EndRenderPass directly as it will mess up the layout manager.
-	void BeginRenderPass(const FVulkanRenderTargetLayout& Layout, class FVulkanRenderPass* RenderPass, class FVulkanFramebuffer* Framebuffer, const VkClearValue* AttachmentClearValues);
+	void BeginRenderPass(const FVulkanBeginRenderPassInfo& BeginRenderPassInfo, const VkClearValue* AttachmentClearValues);
 	void EndRenderPass();
 
 	// Split barriers
diff --git a/Engine/Source/Runtime/VulkanRHI/Private/VulkanDevice.cpp b/Engine/Source/Runtime/VulkanRHI/Private/VulkanDevice.cpp
--- a/Engine/Source/Runtime/VulkanRHI/Private/VulkanDevice.cpp
+++ b/Engine/Source/Runtime/VulkanRHI/Private/VulkanDevice.cpp
@@ -1397,7 +1397,7 @@ void FVulkanDevice::InitGPU()
 	}
 #endif
 
-	RenderPassManager = new FVulkanRenderPassManager(this);
+	RenderPassManager = new FVulkanRenderPassManager(*this);
 
 	if (UseVulkanDescriptorCache())
 	{
diff --git a/Engine/Source/Runtime/VulkanRHI/Private/VulkanRenderTarget.cpp b/Engine/Source/Runtime/VulkanRHI/Private/VulkanRenderTarget.cpp
--- a/Engine/Source/Runtime/VulkanRHI/Private/VulkanRenderTarget.cpp
+++ b/Engine/Source/Runtime/VulkanRHI/Private/VulkanRenderTarget.cpp
@@ -569,7 +569,10 @@ void FVulkanCommandListContext::RHIBeginRenderPass(const FRHIRenderPassInfo& InI
 
 	FVulkanFramebuffer* Framebuffer = Device.GetRenderPassManager().GetOrCreateFramebuffer(RTInfo, RTLayout, RenderPass);
 	checkf(RenderPass != nullptr && Framebuffer != nullptr, TEXT("RenderPass not started! Bad combination of values? Depth %p #Color %d Color0 %p"), (void*)InInfo.DepthStencilRenderTarget.DepthStencilTarget, InInfo.GetNumColorRenderTargets(), (void*)InInfo.ColorRenderTargets[0].RenderTarget);
-	Device.GetRenderPassManager().BeginRenderPass(*this, Device, InInfo, RTLayout, RenderPass, Framebuffer);
+
+	const bool bIsParallelRenderPass = false;
+	FVulkanBeginRenderPassInfo BeginRenderPassInfo{*RenderPass, *Framebuffer, bIsParallelRenderPass};
+	Device.GetRenderPassManager().BeginRenderPass(*this, InInfo, RTLayout, BeginRenderPassInfo);
 
 	check(!CurrentRenderPass);
 	CurrentRenderPass = RenderPass;
diff --git a/Engine/Source/Runtime/VulkanRHI/Private/VulkanRenderpass.cpp b/Engine/Source/Runtime/VulkanRHI/Private/VulkanRenderpass.cpp
--- a/Engine/Source/Runtime/VulkanRHI/Private/VulkanRenderpass.cpp
+++ b/Engine/Source/Runtime/VulkanRHI/Private/VulkanRenderpass.cpp
@@ -42,7 +42,7 @@ FVulkanRenderPassManager::~FVulkanRenderPassManager()
 		FFramebufferList* List = Pair.Value;
 		for (int32 Index = List->Framebuffer.Num() - 1; Index >= 0; --Index)
 		{
-			List->Framebuffer[Index]->Destroy(*Device);
+			List->Framebuffer[Index]->Destroy(Device);
 			delete List->Framebuffer[Index];
 		}
 		delete List;
@@ -121,17 +121,17 @@ FVulkanFramebuffer* FVulkanRenderPassManager::GetOrCreateFramebuffer(const FRHIS
 		}
 	}
 
-	FVulkanFramebuffer* Framebuffer = new FVulkanFramebuffer(*Device, RenderTargetsInfo, RTLayout, *RenderPass);
+	FVulkanFramebuffer* Framebuffer = new FVulkanFramebuffer(Device, RenderTargetsInfo, RTLayout, *RenderPass);
 	FramebufferList->Framebuffer.Add(Framebuffer);
 	return Framebuffer;
 }
 
-void FVulkanRenderPassManager::BeginRenderPass(FVulkanCommandListContext& Context, FVulkanDevice& InDevice, const FRHIRenderPassInfo& RPInfo, const FVulkanRenderTargetLayout& RTLayout, FVulkanRenderPass* RenderPass, FVulkanFramebuffer* Framebuffer)
+void FVulkanRenderPassManager::BeginRenderPass(FVulkanCommandListContext& Context, const FRHIRenderPassInfo& RPInfo, const FVulkanRenderTargetLayout& RTLayout, const FVulkanBeginRenderPassInfo& BeginRenderPassInfo)
 {
 	// (NumRT + 1 [Depth] ) * 2 [surface + resolve]
 	VkClearValue ClearValues[(MaxSimultaneousRenderTargets + 1) * 2];
 	uint32 ClearValueIndex = 0;
-	bool bNeedsClearValues = RenderPass->GetNumUsedClearValues() > 0;
+	bool bNeedsClearValues = BeginRenderPassInfo.RenderPass.GetNumUsedClearValues() > 0;
 	FMemory::Memzero(ClearValues);
 
 	int32 NumColorTargets = RPInfo.GetNumColorRenderTargets();
@@ -202,13 +202,13 @@ void FVulkanRenderPassManager::BeginRenderPass(FVulkanCommandListContext& Contex
 		ValidateShadingRateDataType();
 	}
 
-	ensure(ClearValueIndex <= RenderPass->GetNumUsedClearValues());
+	ensure(ClearValueIndex <= BeginRenderPassInfo.RenderPass.GetNumUsedClearValues());
 
 	FVulkanCommandBuffer& CommandBuffer = Context.GetCommandBuffer();
 
 	Barrier.Execute(&CommandBuffer);
 
-	CommandBuffer.BeginRenderPass(RenderPass->GetLayout(), RenderPass, Framebuffer, ClearValues);
+	CommandBuffer.BeginRenderPass(BeginRenderPassInfo, ClearValues);
 
 	{
 		const VkExtent3D& Extents = RTLayout.GetExtent3D();
@@ -235,7 +235,7 @@ void FVulkanRenderPassManager::NotifyDeletedRenderTarget(VkImage Image)
 			if (Framebuffer->ContainsRenderTarget(Image))
 			{
 				List->Framebuffer.RemoveAtSwap(Index, EAllowShrinking::No);
-				Framebuffer->Destroy(*Device);
+				Framebuffer->Destroy(Device);
 				delete Framebuffer;
 			}
 		}
diff --git a/Engine/Source/Runtime/VulkanRHI/Private/VulkanRenderpass.h b/Engine/Source/Runtime/VulkanRHI/Private/VulkanRenderpass.h
--- a/Engine/Source/Runtime/VulkanRHI/Private/VulkanRenderpass.h
+++ b/Engine/Source/Runtime/VulkanRHI/Private/VulkanRenderpass.h
@@ -852,12 +852,17 @@ private:
 
 VkRenderPass CreateVulkanRenderPass(FVulkanDevice& Device, const FVulkanRenderTargetLayout& RTLayout);
 
+struct FVulkanBeginRenderPassInfo
+{
+	FVulkanRenderPass& RenderPass;
+	FVulkanFramebuffer& Framebuffer;
+	bool bIsParallelRenderPass;
+};
 
-
-class FVulkanRenderPassManager : public VulkanRHI::FDeviceChild
+class FVulkanRenderPassManager
 {
 public:
-	FVulkanRenderPassManager(FVulkanDevice* InDevice) : VulkanRHI::FDeviceChild(InDevice) {}
+	FVulkanRenderPassManager(FVulkanDevice& InDevice) : Device(InDevice) {}
 	~FVulkanRenderPassManager();
 
 	FVulkanFramebuffer* GetOrCreateFramebuffer(const FRHISetRenderTargetsInfo& RenderTargetsInfo, const FVulkanRenderTargetLayout& RTLayout, FVulkanRenderPass* RenderPass);
@@ -874,7 +879,7 @@ public:
 			}
 		}
 
-		FVulkanRenderPass* RenderPass = new FVulkanRenderPass(*Device, RTLayout);
+		FVulkanRenderPass* RenderPass = new FVulkanRenderPass(Device, RTLayout);
 		{
 			FRWScopeLock ScopedWriteLock(RenderPassesLock, SLT_Write);
 			FVulkanRenderPass** FoundRenderPass = RenderPasses.Find(RenderPassHash);
@@ -888,7 +893,7 @@ public:
 		return RenderPass;
 	}
 
-	void BeginRenderPass(FVulkanCommandListContext& Context, FVulkanDevice& InDevice, const FRHIRenderPassInfo& RPInfo, const FVulkanRenderTargetLayout& RTLayout, FVulkanRenderPass* RenderPass, FVulkanFramebuffer* Framebuffer);
+	void BeginRenderPass(FVulkanCommandListContext& Context, const FRHIRenderPassInfo& RPInfo, const FVulkanRenderTargetLayout& RTLayout, const FVulkanBeginRenderPassInfo& BeginRenderPassInfo);
 	void EndRenderPass(FVulkanCommandListContext& Context);
 
 	FRWLock RenderPassesLock;
@@ -905,4 +910,6 @@ private:
 		TArray<FVulkanFramebuffer*> Framebuffer;
 	};
 	TMap<uint32, FFramebufferList*> Framebuffers;
+
+	FVulkanDevice& Device;
 };
