PR #13111: VulkanRHI : Fixed memory leak in VulkanTexture
#rn ElectraPlayer for Android now supports a faster output path on Vulkan that can be enabled by setting CVar Electra.AndroidUseGpuOutputPath to 1

diff --git a/Engine/Source/Runtime/VulkanRHI/Private/VulkanTexture.cpp b/Engine/Source/Runtime/VulkanRHI/Private/VulkanTexture.cpp
--- a/Engine/Source/Runtime/VulkanRHI/Private/VulkanTexture.cpp
+++ b/Engine/Source/Runtime/VulkanRHI/Private/VulkanTexture.cpp
@@ -1834,6 +1834,7 @@ FVulkanTexture::FVulkanTexture(FVulkanDevice& InDevice, const FRHITextureCreateD
 		VulkanImage,
 		VulkanDeviceMemory,
 		SamplerYcbcrConversion,
+		HardwareBuffer
 	};
 
 	ExternalImageDeleteCallbackInfo =
