VulkanRHI: Fix timestamp calibration for new profiler in Android

#rnx

diff --git a/Engine/Source/Runtime/VulkanRHI/Private/VulkanDevice.cpp b/Engine/Source/Runtime/VulkanRHI/Private/VulkanDevice.cpp
--- a/Engine/Source/Runtime/VulkanRHI/Private/VulkanDevice.cpp
+++ b/Engine/Source/Runtime/VulkanRHI/Private/VulkanDevice.cpp
@@ -429,7 +429,7 @@ void FVulkanDevice::CreateDevice(TArray<const ANSICHAR*>& DeviceLayers, FVulkanD
 	}
 	else
 	{
-			bAsyncComputeQueue = true;
+		bAsyncComputeQueue = true;
 		ActiveQueueFamilies.Add(ComputeQueueFamilyIndex);
 	}
 	Queues[(int32)EVulkanQueueType::AsyncCompute] = new FVulkanQueue(*this, ComputeQueueFamilyIndex, EVulkanQueueType::AsyncCompute);
@@ -1693,8 +1693,12 @@ void FVulkanDevice::GetCalibrationTimestamp(FVulkanTiming& InOutTiming)
 		switch (CalibratedTimestampInfo[1].timeDomain)
 		{
 		case VK_TIME_DOMAIN_CLOCK_MONOTONIC_KHR:
-			// Divide to match values returned by Cycles64()
-			InOutTiming.CPUTimestamp = Timestamps[1] / 100ULL;
+			// Divide to match values returned by Cycles64() (to be moved to FVulkanPlatform)
+#if PLATFORM_ANDROID
+			InOutTiming.CPUTimestamp = Timestamps[1] / 1000ULL;  // Android Cycle64 divides by 1000
+#else
+			InOutTiming.CPUTimestamp = Timestamps[1] / 100ULL;  // Linux Cycle64 divides by 100
+#endif
 			break;
 
 		case VK_TIME_DOMAIN_QUERY_PERFORMANCE_COUNTER_KHR:
