VulkanRHI: Add support for parallel rendering contexts using secondary command buffers.

#rnx

diff --git a/Engine/Source/Runtime/VulkanRHI/Private/VulkanCommandBuffer.cpp b/Engine/Source/Runtime/VulkanRHI/Private/VulkanCommandBuffer.cpp
--- a/Engine/Source/Runtime/VulkanRHI/Private/VulkanCommandBuffer.cpp
+++ b/Engine/Source/Runtime/VulkanRHI/Private/VulkanCommandBuffer.cpp
@@ -6,10 +6,10 @@
 
 #include "VulkanCommandBuffer.h"
 #include "VulkanContext.h"
+#include "VulkanRenderpass.h"
 #include "VulkanDescriptorSets.h"
 #include "VulkanMemory.h"
 #include "VulkanRayTracing.h"
-#include "VulkanRenderpass.h"
 
 
 #define CMD_BUFFER_TIME_TO_WAIT_BEFORE_DELETING		10
@@ -35,7 +35,7 @@ void FVulkanCommandBuffer::AllocMemory()
 
 	VkCommandBufferAllocateInfo CreateCmdBufInfo;
 	ZeroVulkanStruct(CreateCmdBufInfo, VK_STRUCTURE_TYPE_COMMAND_BUFFER_ALLOCATE_INFO);
-	CreateCmdBufInfo.level = VK_COMMAND_BUFFER_LEVEL_PRIMARY;
+	CreateCmdBufInfo.level = (GetCommandBufferType() == EVulkanCommandBufferType::Primary) ? VK_COMMAND_BUFFER_LEVEL_PRIMARY : VK_COMMAND_BUFFER_LEVEL_SECONDARY;
 	CreateCmdBufInfo.commandBufferCount = 1;
 	CreateCmdBufInfo.commandPool = CommandBufferPool.GetHandle();
 
@@ -92,16 +92,19 @@ void FVulkanCommandBuffer::BeginRenderPass(const FVulkanBeginRenderPassInfo& Beg
 	Info.clearValueCount = Layout.GetNumUsedClearValues();
 	Info.pClearValues = AttachmentClearValues;
 
+	const VkSubpassContents SubpassContents = BeginRenderPassInfo.bIsParallelRenderPass ?
+		VK_SUBPASS_CONTENTS_SECONDARY_COMMAND_BUFFERS : VK_SUBPASS_CONTENTS_INLINE;
+
 	if (Device.GetOptionalExtensions().HasKHRRenderPass2)
 	{
 		VkSubpassBeginInfo SubpassInfo;
 		ZeroVulkanStruct(SubpassInfo, VK_STRUCTURE_TYPE_SUBPASS_BEGIN_INFO);
-		SubpassInfo.contents = VK_SUBPASS_CONTENTS_INLINE;
+		SubpassInfo.contents = SubpassContents;
 		VulkanRHI::vkCmdBeginRenderPass2KHR(CommandBufferHandle, &Info, &SubpassInfo);
 	}
 	else
 	{
-		VulkanRHI::vkCmdBeginRenderPass(CommandBufferHandle, &Info, VK_SUBPASS_CONTENTS_INLINE);
+		VulkanRHI::vkCmdBeginRenderPass(CommandBufferHandle, &Info, SubpassContents);
 	}
 
 	State = EState::IsInsideRenderPass;
@@ -142,8 +145,11 @@ void FVulkanCommandBuffer::End(FVulkanQueryPool* QueryPool)
 	State = EState::HasEnded;
 }
 
-void FVulkanCommandBuffer::Begin(FVulkanQueryPool* QueryPool)
+void FVulkanCommandBuffer::Begin(FVulkanQueryPool* QueryPool, VkRenderPass RenderPassHandle)
 {
+	checkf((GetCommandBufferType() == EVulkanCommandBufferType::Primary) || (RenderPassHandle != VK_NULL_HANDLE),
+		TEXT("Secondary command buffers require the render pass handle!"));
+
 	{
 		FScopeLock ScopeLock(CommandBufferPool.GetCS());
 		if (State == EState::NeedReset)
@@ -157,11 +163,21 @@ void FVulkanCommandBuffer::Begin(FVulkanQueryPool* QueryPool)
 		State = EState::IsInsideBegin;
 	}
 
-	VkCommandBufferBeginInfo CmdBufBeginInfo;
-	ZeroVulkanStruct(CmdBufBeginInfo, VK_STRUCTURE_TYPE_COMMAND_BUFFER_BEGIN_INFO);
-	CmdBufBeginInfo.flags = VK_COMMAND_BUFFER_USAGE_ONE_TIME_SUBMIT_BIT;
+	VkCommandBufferBeginInfo BeginInfo;
+	ZeroVulkanStruct(BeginInfo, VK_STRUCTURE_TYPE_COMMAND_BUFFER_BEGIN_INFO);
+	BeginInfo.flags = VK_COMMAND_BUFFER_USAGE_ONE_TIME_SUBMIT_BIT;
 
-	VERIFYVULKANRESULT(VulkanRHI::vkBeginCommandBuffer(CommandBufferHandle, &CmdBufBeginInfo));
+	VkCommandBufferInheritanceInfo InheritanceInfo;
+	if (GetCommandBufferType() == EVulkanCommandBufferType::Secondary)
+	{
+		BeginInfo.flags |= VK_COMMAND_BUFFER_USAGE_RENDER_PASS_CONTINUE_BIT;
+		BeginInfo.pInheritanceInfo = &InheritanceInfo;
+
+		ZeroVulkanStruct(InheritanceInfo, VK_STRUCTURE_TYPE_COMMAND_BUFFER_INHERITANCE_INFO);
+		InheritanceInfo.renderPass = RenderPassHandle;
+	}
+
+	VERIFYVULKANRESULT(VulkanRHI::vkBeginCommandBuffer(CommandBufferHandle, &BeginInfo));
 
 	check(!CurrentDescriptorPoolSetContainer);
 
@@ -226,6 +242,14 @@ bool FVulkanCommandBuffer::AcquirePoolSetAndDescriptorsIfNeeded(const class FVul
 
 void FVulkanCommandBuffer::Reset()
 {
+	// Reset the secondary command buffers we executed from this one
+	for (FVulkanCommandBuffer* SecondaryCommand : ExecutedSecondaryCommandBuffers)
+	{
+		SecondaryCommand->Reset();
+	}
+	ExecutedSecondaryCommandBuffers.Empty();
+
+	// Hold lock while State is altered
 	FScopeLock ScopeLock(CommandBufferPool.GetCS());
 	if (State == EState::Submitted)
 	{
@@ -280,9 +304,10 @@ void FVulkanCommandBuffer::EndSplitBarrier(VkEvent BarrierEvent, const VkDepende
 }
 
 
-FVulkanCommandBufferPool::FVulkanCommandBufferPool(FVulkanDevice& InDevice, FVulkanQueue& InQueue)
+FVulkanCommandBufferPool::FVulkanCommandBufferPool(FVulkanDevice& InDevice, FVulkanQueue& InQueue, EVulkanCommandBufferType InCommandBufferType)
 	: Device(InDevice)
 	, Queue(InQueue)
+	, CommandBufferType(InCommandBufferType)
 {
 	VkCommandPoolCreateInfo CmdPoolInfo;
 	ZeroVulkanStruct(CmdPoolInfo, VK_STRUCTURE_TYPE_COMMAND_POOL_CREATE_INFO);
diff --git a/Engine/Source/Runtime/VulkanRHI/Private/VulkanCommandBuffer.h b/Engine/Source/Runtime/VulkanRHI/Private/VulkanCommandBuffer.h
--- a/Engine/Source/Runtime/VulkanRHI/Private/VulkanCommandBuffer.h
+++ b/Engine/Source/Runtime/VulkanRHI/Private/VulkanCommandBuffer.h
@@ -25,6 +25,13 @@ namespace VulkanRHI
 	class FSemaphore;
 }
 
+enum class EVulkanCommandBufferType : uint8
+{
+	Primary,
+	Secondary,
+	Count
+};
+
 class FVulkanCommandBuffer
 {
 protected:
@@ -65,14 +72,16 @@ public:
 		return State != EState::NotAllocated;
 	}
 
-	VkCommandBuffer GetHandle()
+	VkCommandBuffer GetHandle() const
 	{
 		return CommandBufferHandle;
 	}
 
+	EVulkanCommandBufferType GetCommandBufferType() const;
+
 	void SetSubmitted();
 
-	void Begin(FVulkanQueryPool* QueryPool);
+	void Begin(FVulkanQueryPool* QueryPool, VkRenderPass RenderPassHandle);
 	void End(FVulkanQueryPool* QueryPool);
 
 	enum class EState : uint8
@@ -96,6 +105,9 @@ public:
 	uint8 bHasScissor					: 1 = 0;
 	uint8 bHasStencilRef				: 1 = 0;
 
+	// Every secondary command buffer executed from this one with tied lifetimes
+	TArray<FVulkanCommandBuffer*> ExecutedSecondaryCommandBuffers;
+
 	// You never want to call Begin/EndRenderPass directly as it will mess up the layout manager.
 	void BeginRenderPass(const FVulkanBeginRenderPassInfo& BeginRenderPassInfo, const VkClearValue* AttachmentClearValues);
 	void EndRenderPass();
@@ -157,8 +169,6 @@ private:
 
 	FVulkanCommandBufferPool& CommandBufferPool;
 
-	uint64 LastValidTiming = 0;
-
 	TArray<VkEvent> EndedBarrierEvents;
 
 	void AcquirePoolSetContainer();
@@ -185,7 +195,7 @@ public:
 class FVulkanCommandBufferPool
 {
 public:
-	FVulkanCommandBufferPool(FVulkanDevice& InDevice, FVulkanQueue& InQueue);
+	FVulkanCommandBufferPool(FVulkanDevice& InDevice, FVulkanQueue& InQueue, EVulkanCommandBufferType InCommandBufferType);
 	~FVulkanCommandBufferPool();
 
 	VkCommandPool GetHandle() const
@@ -203,12 +213,19 @@ public:
 		return Queue;
 	}
 
+	EVulkanCommandBufferType GetCommandBufferType() const
+	{
+		return CommandBufferType;
+	}
+
 	void FreeUnusedCmdBuffers(FVulkanQueue* Queue, bool bTrimMemory);
 
 private:
 	FVulkanDevice& Device;
 	FVulkanQueue& Queue;
 
+	const EVulkanCommandBufferType CommandBufferType;
+
 	VkCommandPool Handle = VK_NULL_HANDLE;
 
 	TArray<FVulkanCommandBuffer*> CmdBuffers;
@@ -220,3 +237,9 @@ private:
 
 	friend FVulkanContextCommon;
 };
+
+
+inline EVulkanCommandBufferType FVulkanCommandBuffer::GetCommandBufferType() const
+{
+	return CommandBufferPool.GetCommandBufferType();
+}
diff --git a/Engine/Source/Runtime/VulkanRHI/Private/VulkanContext.cpp b/Engine/Source/Runtime/VulkanRHI/Private/VulkanContext.cpp
--- a/Engine/Source/Runtime/VulkanRHI/Private/VulkanContext.cpp
+++ b/Engine/Source/Runtime/VulkanRHI/Private/VulkanContext.cpp
@@ -12,10 +12,10 @@
 
 
 
-FVulkanContextCommon::FVulkanContextCommon(FVulkanDevice& InDevice, FVulkanQueue& InQueue)
+FVulkanContextCommon::FVulkanContextCommon(FVulkanDevice& InDevice, FVulkanQueue& InQueue, EVulkanCommandBufferType InCommandBufferType)
 	: Device(InDevice)
 	, Queue(InQueue)
-	, Pool(InDevice, InQueue)
+	, Pool(InDevice, InQueue, InCommandBufferType)
 {
 }
 
@@ -169,12 +169,14 @@ void FVulkanContextCommon::PrepareNewCommandBuffer(FVulkanPayload& Payload)
 
 	Payload.CommandBuffers.Add(NewCommandBuffer);
 
-	NewCommandBuffer->Begin(GRHIGlobals.SupportsTimestampRenderQueries ? GetCurrentTimestampQueryPool(Payload) : nullptr);
+	FVulkanQueryPool* TimestampQueryPool = GRHIGlobals.SupportsTimestampRenderQueries ? GetCurrentTimestampQueryPool(Payload) : nullptr;
+	const VkRenderPass RenderPassHandle = GetParallelRenderPassInfo() ? GetParallelRenderPassInfo()->RenderPassHandle : VK_NULL_HANDLE;
+	NewCommandBuffer->Begin(TimestampQueryPool, RenderPassHandle);
 }
 
 
 FVulkanCommandListContext::FVulkanCommandListContext(FVulkanDevice& InDevice, ERHIPipeline InPipeline, FVulkanCommandListContext* InImmediate)
-	: FVulkanContextCommon(InDevice, *InDevice.GetQueue(InPipeline))
+	: FVulkanContextCommon(InDevice, *InDevice.GetQueue(InPipeline), EVulkanCommandBufferType::Primary)
 	, Immediate(InImmediate)
 	, RHIPipeline(InPipeline)
 #if (RHI_NEW_GPU_PROFILER == 0)
@@ -193,6 +195,22 @@ FVulkanCommandListContext::FVulkanCommandListContext(FVulkanDevice& InDevice, ER
 	GlobalUniformBuffers.AddZeroed(FUniformBufferStaticSlotRegistry::Get().GetSlotCount());
 }
 
+FVulkanCommandListContext::FVulkanCommandListContext(FVulkanDevice& InDevice, FVulkanCommandListContext* InImmediate, FVulkanParallelRenderPassInfo* InParallelRenderPassInfo)
+	: FVulkanContextCommon(InDevice, *InDevice.GetQueue(ERHIPipeline::Graphics), EVulkanCommandBufferType::Secondary)
+	, Immediate(InImmediate)
+	, RHIPipeline(ERHIPipeline::Graphics)
+	, CurrentParallelRenderPassInfo(InParallelRenderPassInfo)
+#if (RHI_NEW_GPU_PROFILER == 0)
+	, GpuProfiler(this, &InDevice)
+#endif
+{
+	// Only graphic commands can be used
+	PendingGfxState = new FVulkanPendingGfxState(Device);
+
+	GlobalUniformBuffers.AddZeroed(FUniformBufferStaticSlotRegistry::Get().GetSlotCount());
+}
+
+
 void FVulkanCommandListContext::ReleasePendingState()
 {
 	delete PendingGfxState;
@@ -229,7 +247,7 @@ void FVulkanUploadContext::DestroyPool()
 
 
 FVulkanUploadContext::FVulkanUploadContext(FVulkanDevice& InDevice, FVulkanQueue& InQueue)
-	: FVulkanContextCommon(InDevice, InQueue)
+	: FVulkanContextCommon(InDevice, InQueue, EVulkanCommandBufferType::Primary)
 {
 
 }
diff --git a/Engine/Source/Runtime/VulkanRHI/Private/VulkanContext.h b/Engine/Source/Runtime/VulkanRHI/Private/VulkanContext.h
--- a/Engine/Source/Runtime/VulkanRHI/Private/VulkanContext.h
+++ b/Engine/Source/Runtime/VulkanRHI/Private/VulkanContext.h
@@ -32,6 +32,11 @@ enum class EVulkanFlushFlags
 ENUM_CLASS_FLAGS(EVulkanFlushFlags)
 
 
+struct FVulkanParallelRenderPassInfo
+{
+	VkRenderPass RenderPassHandle = VK_NULL_HANDLE;
+	TArray<FVulkanPayload*> SecondaryPayloads;
+};
 
 
 class FVulkanContextCommon
@@ -48,7 +53,7 @@ class FVulkanContextCommon
 	};
 
 public:
-	FVulkanContextCommon(FVulkanDevice& InDevice, FVulkanQueue& InQueue);
+	FVulkanContextCommon(FVulkanDevice& InDevice, FVulkanQueue& InQueue, EVulkanCommandBufferType InCommandBufferType);
 	virtual ~FVulkanContextCommon();
 
 	FVulkanPayload& GetPayload(EPhase Phase)
@@ -185,6 +190,11 @@ public:
 		return GetCurrentTimestampQueryPool(GetPayload(EPhase::Execute));
 	}
 
+	virtual FVulkanParallelRenderPassInfo* GetParallelRenderPassInfo()
+	{
+		return nullptr;
+	}
+
 private:
 	FVulkanQueryPool* GetCurrentTimestampQueryPool(FVulkanPayload& Payload);
 	void PrepareNewCommandBuffer(FVulkanPayload& Payload);
@@ -211,6 +221,9 @@ public:
 	FVulkanCommandListContext(FVulkanDevice& InDevice, ERHIPipeline InPipeline, FVulkanCommandListContext* InImmediate);
 	virtual ~FVulkanCommandListContext();
 
+	// Constructor for parallel render contexts that use secondary command buffers
+	FVulkanCommandListContext(FVulkanDevice& InDevice, FVulkanCommandListContext* InImmediate, FVulkanParallelRenderPassInfo* InParallelRenderPassInfo);
+
 	static inline FVulkanCommandListContext& Get(FRHICommandListBase& RHICmdList)
 	{
 		return static_cast<FVulkanCommandListContext&>(RHICmdList.GetContext().GetLowestLevelContext());
@@ -372,6 +385,16 @@ public:
 
 	FVulkanQueryPool* GetCurrentOcclusionQueryPool();
 
+	virtual FVulkanParallelRenderPassInfo* GetParallelRenderPassInfo()
+	{
+		return CurrentParallelRenderPassInfo;
+	}
+	
+	void SetParallelRenderPassInfo(FVulkanParallelRenderPassInfo* ParallelRenderPassInfo)
+	{
+		CurrentParallelRenderPassInfo = ParallelRenderPassInfo;
+	}
+
 protected:
 	FVulkanCommandListContext* Immediate;
 	const ERHIPipeline RHIPipeline;
@@ -383,6 +406,8 @@ protected:
 	FVulkanRenderPass* CurrentRenderPass = nullptr;
 	FVulkanFramebuffer* CurrentFramebuffer = nullptr;
 
+	FVulkanParallelRenderPassInfo* CurrentParallelRenderPassInfo = nullptr;
+
 	FVulkanPendingGfxState* PendingGfxState = nullptr;
 	FVulkanPendingComputeState* PendingComputeState = nullptr;
 
diff --git a/Engine/Source/Runtime/VulkanRHI/Private/VulkanDevice.cpp b/Engine/Source/Runtime/VulkanRHI/Private/VulkanDevice.cpp
--- a/Engine/Source/Runtime/VulkanRHI/Private/VulkanDevice.cpp
+++ b/Engine/Source/Runtime/VulkanRHI/Private/VulkanDevice.cpp
@@ -1617,10 +1617,9 @@ void FVulkanDevice::NotifyDeletedComputePipeline(class FVulkanComputePipeline* P
 	}
 }
 
-static FCriticalSection GContextCS;
 FVulkanCommandListContext* FVulkanDevice::AcquireDeferredContext(ERHIPipeline Pipeline)
 {
-	FScopeLock Lock(&GContextCS);
+	FScopeLock Lock(&ContextCS);
 
 	TArray<FVulkanCommandListContext*>& Contexts = CommandContexts[Pipeline];
 
@@ -1635,11 +1634,33 @@ FVulkanCommandListContext* FVulkanDevice::AcquireDeferredContext(ERHIPipeline Pi
 void FVulkanDevice::ReleaseDeferredContext(FVulkanCommandListContext* InContext)
 {
 	check(InContext);
+	FScopeLock Lock(&ContextCS);
+	TArray<FVulkanCommandListContext*>& Contexts = CommandContexts[InContext->GetPipeline()];
+	Contexts.Push(InContext);
+}
+
+FVulkanCommandListContext* FVulkanDevice::AcquireParallelRenderContext(FVulkanParallelRenderPassInfo* ParallelRenderPassInfo)
+{
+	FScopeLock Lock(&ContextCS);
+	FVulkanCommandListContext* ParallelRenderContext = nullptr;
+	if (ParallelRenderContexts.Num() > 0)
+	{
+		ParallelRenderContext = ParallelRenderContexts.Pop(EAllowShrinking::No);
+		ParallelRenderContext->SetParallelRenderPassInfo(ParallelRenderPassInfo);
+	}
+	else
 	{
-		FScopeLock Lock(&GContextCS);
-		TArray<FVulkanCommandListContext*>& Contexts = CommandContexts[InContext->GetPipeline()];
-		Contexts.Push(InContext);
+		ParallelRenderContext = new FVulkanCommandListContext(*this, ImmediateContext, ParallelRenderPassInfo);
 	}
+	return ParallelRenderContext;
+}
+
+void FVulkanDevice::ReleaseParallelRenderContext(FVulkanCommandListContext* InContext)
+{
+	check(InContext);
+	FScopeLock Lock(&ContextCS);
+	InContext->SetParallelRenderPassInfo(nullptr);
+	ParallelRenderContexts.Push(InContext);
 }
 
 void FVulkanDevice::VulkanSetObjectName(VkObjectType Type, uint64_t Handle, const TCHAR* Name)
diff --git a/Engine/Source/Runtime/VulkanRHI/Private/VulkanDevice.h b/Engine/Source/Runtime/VulkanRHI/Private/VulkanDevice.h
--- a/Engine/Source/Runtime/VulkanRHI/Private/VulkanDevice.h
+++ b/Engine/Source/Runtime/VulkanRHI/Private/VulkanDevice.h
@@ -24,7 +24,7 @@ class FVulkanDeviceExtension;
 class FVulkanRenderPassManager;
 class FVulkanRayTracingCompactionRequestHandler;
 class FVulkanRayTracingPipelineLibraryCache;
-
+struct FVulkanParallelRenderPassInfo;
 
 #define VULKAN_USE_DEBUG_NAMES 1
 #if VULKAN_USE_DEBUG_NAMES
@@ -512,6 +512,10 @@ public:
 
 	FVulkanCommandListContext* AcquireDeferredContext(ERHIPipeline Pipeline);
 	void ReleaseDeferredContext(FVulkanCommandListContext* InContext);
+
+	FVulkanCommandListContext* AcquireParallelRenderContext(FVulkanParallelRenderPassInfo* ParallelRenderPassInfo);
+	void ReleaseParallelRenderContext(FVulkanCommandListContext* InContext);
+
 	void VulkanSetObjectName(VkObjectType Type, uint64_t Handle, const TCHAR* Name);
 	const FOptionalVulkanDeviceExtensions& GetOptionalExtensions() const
 	{
@@ -663,8 +667,10 @@ private:
 
 	TMap<uint32, FSamplerStateRHIRef> SamplerMap;
 
+	FCriticalSection ContextCS;
 	FVulkanCommandListContextImmediate* ImmediateContext;
 	TRHIPipelineArray<TArray<FVulkanCommandListContext*>> CommandContexts;
+	TArray<FVulkanCommandListContext*> ParallelRenderContexts;
 
 	FVulkanDynamicRHI* RHI = nullptr;
 	
