VulkanRHI: Add support for RHI parallel render passes and enable parallel translate on desktop.

#rnx

diff --git a/Engine/Source/Runtime/VulkanRHI/Private/VulkanContext.cpp b/Engine/Source/Runtime/VulkanRHI/Private/VulkanContext.cpp
--- a/Engine/Source/Runtime/VulkanRHI/Private/VulkanContext.cpp
+++ b/Engine/Source/Runtime/VulkanRHI/Private/VulkanContext.cpp
@@ -234,7 +234,55 @@ FVulkanCommandListContext::~FVulkanCommandListContext()
 	ReleasePendingState();
 }
 
+void FVulkanCommandListContext::RHIBeginParallelRenderPass(TSharedPtr<FRHIParallelRenderPassInfo> InInfo, const TCHAR* InName)
+{
+	checkf(CurrentParallelRenderPassInfo == nullptr, TEXT("There is already a parallel render pass in progress!"));
+	CurrentParallelRenderPassInfo = new FVulkanParallelRenderPassInfo();
+
+	RHIBeginRenderPass(*InInfo.Get(), InName);
+
+	CurrentParallelRenderPassInfo->RenderPassHandle = CurrentRenderPass->GetHandle();
+	InInfo->RHIPlatformData = CurrentParallelRenderPassInfo;
+}
+
+void FVulkanCommandListContext::RHIEndParallelRenderPass()
+{
+	if (CurrentParallelRenderPassInfo->SecondaryPayloads.Num())
+	{
+		FVulkanCommandBuffer& ParentCommandBuffer = GetCommandBuffer();
+
+		TArray<VkCommandBuffer> CommandBufferHandles;
+		CommandBufferHandles.Reserve(CurrentParallelRenderPassInfo->SecondaryPayloads.Num()); // should be 1 cmdbuffer per payload
+		for (FVulkanPayload* Payload : CurrentParallelRenderPassInfo->SecondaryPayloads)
+		{
+			check(Payload->SignalSemaphores.Num() == 0);
+			check(Payload->WaitSemaphores.Num() == 0);
+
+			for (FVulkanCommandBuffer* SecondaryCommandBuffer : Payload->CommandBuffers)
+			{
+				CommandBufferHandles.Add(SecondaryCommandBuffer->GetHandle());
+			}
+			ParentCommandBuffer.ExecutedSecondaryCommandBuffers.Append(MoveTemp(Payload->CommandBuffers));
 
+			for (const FVulkanSyncPointRef& Sync : Payload->SyncPoints)
+			{
+				AddPendingSyncPoint(Sync);
+			}
+
+			delete Payload;
+		}
+
+		if (CommandBufferHandles.Num())
+		{
+			VulkanRHI::vkCmdExecuteCommands(ParentCommandBuffer.GetHandle(), CommandBufferHandles.Num(), CommandBufferHandles.GetData());
+		}
+	}
+
+	RHIEndRenderPass();
+
+	delete CurrentParallelRenderPassInfo;
+	CurrentParallelRenderPassInfo = nullptr;
+}
 
 TLockFreePointerListUnordered<FVulkanUploadContext, PLATFORM_CACHE_LINE_SIZE> FVulkanUploadContext::Pool;
 void FVulkanUploadContext::DestroyPool()
diff --git a/Engine/Source/Runtime/VulkanRHI/Private/VulkanContext.h b/Engine/Source/Runtime/VulkanRHI/Private/VulkanContext.h
--- a/Engine/Source/Runtime/VulkanRHI/Private/VulkanContext.h
+++ b/Engine/Source/Runtime/VulkanRHI/Private/VulkanContext.h
@@ -321,6 +321,9 @@ public:
 	virtual void RHIEndRenderPass() final override;
 	virtual void RHINextSubpass() final override;
 
+	virtual void RHIBeginParallelRenderPass(TSharedPtr<FRHIParallelRenderPassInfo> InInfo, const TCHAR* InName) final override;
+	virtual void RHIEndParallelRenderPass() final override;
+
 	virtual void RHIClearRayTracingBindings(FRHIRayTracingScene* Scene) final override;
 	virtual void RHICommitRayTracingBindings(FRHIRayTracingScene* Scene) final override;
 	virtual void RHIClearShaderBindingTable(FRHIShaderBindingTable* SBT) final override;
diff --git a/Engine/Source/Runtime/VulkanRHI/Private/VulkanDynamicRHI.h b/Engine/Source/Runtime/VulkanRHI/Private/VulkanDynamicRHI.h
--- a/Engine/Source/Runtime/VulkanRHI/Private/VulkanDynamicRHI.h
+++ b/Engine/Source/Runtime/VulkanRHI/Private/VulkanDynamicRHI.h
@@ -242,6 +242,9 @@ public:
 	virtual void RHISubmitCommandLists(FRHISubmitCommandListsArgs&& Args) final override;
 	virtual uint64 RHIGetMinimumAlignmentForBufferBackedSRV(EPixelFormat Format) final override;
 
+	virtual IRHIComputeContext* RHIGetParallelCommandContext(FRHIParallelRenderPassInfo const& ParallelRenderPass, FRHIGPUMask GPUMask) final override;
+	virtual IRHIPlatformCommandList* RHIFinalizeParallelContext(IRHIComputeContext* Context) final override;
+
 	// Transient Resource Functions and Helpers
 	virtual IRHITransientResourceAllocator* RHICreateTransientResourceAllocator() final override;
 
diff --git a/Engine/Source/Runtime/VulkanRHI/Private/VulkanRHI.cpp b/Engine/Source/Runtime/VulkanRHI/Private/VulkanRHI.cpp
--- a/Engine/Source/Runtime/VulkanRHI/Private/VulkanRHI.cpp
+++ b/Engine/Source/Runtime/VulkanRHI/Private/VulkanRHI.cpp
@@ -63,7 +63,7 @@ TAtomic<uint64> GVulkanDSetLayoutHandleIdCounter{ 0 };
 
 TAutoConsoleVariable<int32> GRHIThreadCvar(
 	TEXT("r.Vulkan.RHIThread"),
-	1,
+	2,
 	TEXT("0 to only use Render Thread\n")
 	TEXT("1 to use ONE RHI Thread\n")
 	TEXT("2 to use multiple RHI Thread\n")
@@ -896,7 +896,6 @@ void FVulkanDynamicRHI::InitInstance()
 			GRHIGlobals.RayTracing.SupportsLooseParamsInShaderRecord = true;
 		}
 
-		GRHISupportsParallelRenderPasses = false;
 #if VULKAN_ENABLE_DUMP_LAYER
 		// Disable RHI thread by default if the dump layer is enabled
 		GRHISupportsRHIThread = false;
@@ -906,6 +905,10 @@ void FVulkanDynamicRHI::InitInstance()
 		GRHISupportsParallelRHIExecute = GRHIThreadCvar->GetInt() > 1;
 #endif
 
+		GRHISupportsParallelRenderPasses = GRHISupportsParallelRHIExecute;
+		GRHIParallelRHIExecuteChildWait = GRHISupportsParallelRHIExecute;
+		GRHIParallelRHIExecuteParentWait = GRHISupportsParallelRHIExecute;
+
 		GRHISupportsUAVFormatAliasing = true;
 
 		// Some platforms might only have CPU for an RHI thread, but not for parallel tasks
diff --git a/Engine/Source/Runtime/VulkanRHI/Private/VulkanRenderTarget.cpp b/Engine/Source/Runtime/VulkanRHI/Private/VulkanRenderTarget.cpp
--- a/Engine/Source/Runtime/VulkanRHI/Private/VulkanRenderTarget.cpp
+++ b/Engine/Source/Runtime/VulkanRHI/Private/VulkanRenderTarget.cpp
@@ -570,7 +570,7 @@ void FVulkanCommandListContext::RHIBeginRenderPass(const FRHIRenderPassInfo& InI
 	FVulkanFramebuffer* Framebuffer = Device.GetRenderPassManager().GetOrCreateFramebuffer(RTInfo, RTLayout, RenderPass);
 	checkf(RenderPass != nullptr && Framebuffer != nullptr, TEXT("RenderPass not started! Bad combination of values? Depth %p #Color %d Color0 %p"), (void*)InInfo.DepthStencilRenderTarget.DepthStencilTarget, InInfo.GetNumColorRenderTargets(), (void*)InInfo.ColorRenderTargets[0].RenderTarget);
 
-	const bool bIsParallelRenderPass = false;
+	const bool bIsParallelRenderPass = (CurrentParallelRenderPassInfo != nullptr);
 	FVulkanBeginRenderPassInfo BeginRenderPassInfo{*RenderPass, *Framebuffer, bIsParallelRenderPass};
 	Device.GetRenderPassManager().BeginRenderPass(*this, InInfo, RTLayout, BeginRenderPassInfo);
 
diff --git a/Engine/Source/Runtime/VulkanRHI/Private/VulkanSubmission.cpp b/Engine/Source/Runtime/VulkanRHI/Private/VulkanSubmission.cpp
--- a/Engine/Source/Runtime/VulkanRHI/Private/VulkanSubmission.cpp
+++ b/Engine/Source/Runtime/VulkanRHI/Private/VulkanSubmission.cpp
@@ -244,6 +244,29 @@ void FVulkanDynamicRHI::RHIFinalizeContext(FRHIFinalizeContextArgs&& Args, TRHIP
 	}
 }
 
+IRHIComputeContext* FVulkanDynamicRHI::RHIGetParallelCommandContext(FRHIParallelRenderPassInfo const& ParallelRenderPass, FRHIGPUMask GPUMask)
+{
+	FVulkanParallelRenderPassInfo* VulkanParallelRenderPassInfo = (FVulkanParallelRenderPassInfo*)ParallelRenderPass.RHIPlatformData;
+	checkf(VulkanParallelRenderPassInfo, TEXT("Must begin parallel render pass on parent context before creating parallel contexts!"));
+	return Device->AcquireParallelRenderContext(VulkanParallelRenderPassInfo);
+}
+
+static FCriticalSection GSecondaryPayloadsCS;
+IRHIPlatformCommandList* FVulkanDynamicRHI::RHIFinalizeParallelContext(IRHIComputeContext* Context)
+{
+	FVulkanCommandListContext* VulkanContext = static_cast<FVulkanCommandListContext*>(Context);
+
+	// Call Finalize under this lock
+	{
+		FScopeLock Lock(&GSecondaryPayloadsCS);
+		FVulkanParallelRenderPassInfo* VulkanParallelRenderPassInfo = VulkanContext->GetParallelRenderPassInfo();
+		VulkanContext->Finalize(VulkanParallelRenderPassInfo->SecondaryPayloads);
+	}
+
+	Device->ReleaseParallelRenderContext(VulkanContext);
+	return nullptr;
+}
+
 bool FVulkanDynamicRHI::ProcessSubmissionQueue()
 {
 	SCOPED_NAMED_EVENT_TEXT("ProcessSubmissionQueue", FColor::Orange);
diff --git a/Engine/Source/Runtime/VulkanRHI/Private/VulkanSubmission.h b/Engine/Source/Runtime/VulkanRHI/Private/VulkanSubmission.h
--- a/Engine/Source/Runtime/VulkanRHI/Private/VulkanSubmission.h
+++ b/Engine/Source/Runtime/VulkanRHI/Private/VulkanSubmission.h
@@ -28,6 +28,7 @@ static FVulkanSyncPointRef CreateVulkanSyncPoint()
 
 class FVulkanPayload
 {
+	friend class FVulkanCommandListContext;
 	friend class FVulkanContextCommon;
 	friend class FVulkanDynamicRHI;
 	friend class FVulkanQueue;
