VulkanRHI: Flip the behavior of the local vulkan loader to optionally opt in with a command line.  Keep it as an option for support, but do not replace potentially newer local versions.

#rnx

diff --git a/Engine/Source/Runtime/VulkanRHI/Private/Linux/VulkanLinuxPlatform.cpp b/Engine/Source/Runtime/VulkanRHI/Private/Linux/VulkanLinuxPlatform.cpp
--- a/Engine/Source/Runtime/VulkanRHI/Private/Linux/VulkanLinuxPlatform.cpp
+++ b/Engine/Source/Runtime/VulkanRHI/Private/Linux/VulkanLinuxPlatform.cpp
@@ -88,7 +88,7 @@ bool FVulkanLinuxPlatform::LoadVulkanLibrary()
 #endif // VULKAN_HAS_DEBUGGING_ENABLED
 
 	// Try to load libvulkan.so from the included SDK
-	if (!FParse::Param(FCommandLine::Get(), TEXT("SkipLocalVulkanLoader")))
+	if (FParse::Param(FCommandLine::Get(), TEXT("UseLocalVulkanLoader")))
 	{
 		FAnsiString VulkanLoaderPath(UEVulkanBinariesPath);
 		VulkanLoaderPath.Append("/libvulkan.so");
