VulkanRHI: Create a FVulkanCommandBufferPool pool in the queue (currently split for Primary and Secondary command buffers) so we can create/destroy non-immediate contexts instead of pooling them.  No changes for non-parallel devices (the immediate context pulls a single FVulkanCommandBufferPool and uses it all along).

#rnx

#tests Win CitySample, Android Lyra

diff --git a/Engine/Source/Runtime/VulkanRHI/Private/VulkanCommands.cpp b/Engine/Source/Runtime/VulkanRHI/Private/VulkanCommands.cpp
--- a/Engine/Source/Runtime/VulkanRHI/Private/VulkanCommands.cpp
+++ b/Engine/Source/Runtime/VulkanRHI/Private/VulkanCommands.cpp
@@ -919,6 +919,5 @@ void FVulkanCommandListContext::RHIWriteGPUFence(FRHIGPUFence* FenceRHI)
 IRHIComputeContext* FVulkanDynamicRHI::RHIGetCommandContext(ERHIPipeline Pipeline, FRHIGPUMask GPUMask)
 {
 	check(IsSingleRHIPipeline(Pipeline));
-	FVulkanCommandListContext* CmdContext = Device->AcquireDeferredContext(Pipeline);
-	return CmdContext;
+	return new FVulkanCommandListContext(*Device, Pipeline, &Device->GetImmediateContext());
 }
diff --git a/Engine/Source/Runtime/VulkanRHI/Private/VulkanContext.cpp b/Engine/Source/Runtime/VulkanRHI/Private/VulkanContext.cpp
--- a/Engine/Source/Runtime/VulkanRHI/Private/VulkanContext.cpp
+++ b/Engine/Source/Runtime/VulkanRHI/Private/VulkanContext.cpp
@@ -15,12 +15,13 @@
 FVulkanContextCommon::FVulkanContextCommon(FVulkanDevice& InDevice, FVulkanQueue& InQueue, EVulkanCommandBufferType InCommandBufferType)
 	: Device(InDevice)
 	, Queue(InQueue)
-	, Pool(InDevice, InQueue, InCommandBufferType)
+	, Pool(*InQueue.AcquireCommandBufferPool(InCommandBufferType))
 {
 }
 
 FVulkanContextCommon::~FVulkanContextCommon()
 {
+	Queue.ReleaseCommandBufferPool(&Pool);
 }
 
 void FVulkanContextCommon::NewPayload()
@@ -204,6 +205,8 @@ FVulkanCommandListContext::FVulkanCommandListContext(FVulkanDevice& InDevice, FV
 	, GpuProfiler(this, &InDevice)
 #endif
 {
+	checkf(CurrentParallelRenderPassInfo, TEXT("Secondary command buffers should be created with a FVulkanParallelRenderPassInfo."));
+
 	// Only graphic commands can be used
 	PendingGfxState = new FVulkanPendingGfxState(Device);
 
@@ -225,9 +228,12 @@ FVulkanCommandListContext::~FVulkanCommandListContext()
 	if (GSupportsTimestampRenderQueries)
 	{
 #if (RHI_NEW_GPU_PROFILER == 0)
-		FrameTiming->Release();
-		delete FrameTiming;
-		FrameTiming = nullptr;
+		if (FrameTiming)
+		{
+			FrameTiming->Release();
+			delete FrameTiming;
+			FrameTiming = nullptr;
+		}
 #endif
 	}
 
diff --git a/Engine/Source/Runtime/VulkanRHI/Private/VulkanContext.h b/Engine/Source/Runtime/VulkanRHI/Private/VulkanContext.h
--- a/Engine/Source/Runtime/VulkanRHI/Private/VulkanContext.h
+++ b/Engine/Source/Runtime/VulkanRHI/Private/VulkanContext.h
@@ -201,7 +201,7 @@ private:
 	void NewPayload();
 	void EndPayload();
 
-	FVulkanCommandBufferPool Pool;
+	FVulkanCommandBufferPool& Pool;
 
 	TArray<FVulkanPayload*> Payloads;
 	EPhase CurrentPhase = EPhase::Wait;
@@ -444,7 +444,7 @@ private:
 	void RegisterGPUDispatch(FIntVector GroupCount)	                        { GpuProfiler.RegisterGPUDispatch(GroupCount); }
 
 	FVulkanGPUProfiler GpuProfiler;
-	FVulkanGPUTiming* FrameTiming;
+	FVulkanGPUTiming* FrameTiming = nullptr;
 #endif
 
 	template <typename TRHIShader>
diff --git a/Engine/Source/Runtime/VulkanRHI/Private/VulkanDevice.cpp b/Engine/Source/Runtime/VulkanRHI/Private/VulkanDevice.cpp
--- a/Engine/Source/Runtime/VulkanRHI/Private/VulkanDevice.cpp
+++ b/Engine/Source/Runtime/VulkanRHI/Private/VulkanDevice.cpp
@@ -1421,17 +1421,6 @@ void FVulkanDevice::InitGPU()
 
 	ImmediateContext = new FVulkanCommandListContextImmediate(*this);
 
-	extern TAutoConsoleVariable<int32> GRHIThreadCvar;
-	if (GRHIThreadCvar->GetInt() > 1)
-	{
-		int32 Num = FTaskGraphInterface::Get().GetNumWorkerThreads();
-		for (int32 Index = 0; Index < Num; Index++)
-		{
-			FVulkanCommandListContext* CmdContext = new FVulkanCommandListContext(*this, ERHIPipeline::Graphics, ImmediateContext);
-			CommandContexts[ERHIPipeline::Graphics].Add(CmdContext);
-		}
-	}
-
 #if VULKAN_SUPPORTS_VALIDATION_CACHE
 	if (OptionalDeviceExtensions.HasEXTValidationCache)
 	{
@@ -1489,18 +1478,6 @@ void FVulkanDevice::Destroy()
 	delete DescriptorPoolsManager;
 	DescriptorPoolsManager = nullptr;
 
-	for (int32 PipelineIndex = 0; PipelineIndex < (int32)ERHIPipeline::Num; ++PipelineIndex)
-	{
-		TArray<FVulkanCommandListContext*>& Contexts = CommandContexts[PipelineIndex];
-
-		for (int32 ContextIndex = Contexts.Num() - 1; ContextIndex >= 0; --ContextIndex)
-		{
-			delete Contexts[ContextIndex];
-		}
-
-		Contexts.Reset();
-	}
-
 	delete ImmediateContext;
 	ImmediateContext = nullptr;
 
@@ -1617,52 +1594,6 @@ void FVulkanDevice::NotifyDeletedComputePipeline(class FVulkanComputePipeline* P
 	}
 }
 
-FVulkanCommandListContext* FVulkanDevice::AcquireDeferredContext(ERHIPipeline Pipeline)
-{
-	FScopeLock Lock(&ContextCS);
-
-	TArray<FVulkanCommandListContext*>& Contexts = CommandContexts[Pipeline];
-
-	if (Contexts.Num() == 0)
-	{
-		return new FVulkanCommandListContext(*this, Pipeline, ImmediateContext);
-	}
-
-	return Contexts.Pop(EAllowShrinking::No);
-}
-
-void FVulkanDevice::ReleaseDeferredContext(FVulkanCommandListContext* InContext)
-{
-	check(InContext);
-	FScopeLock Lock(&ContextCS);
-	TArray<FVulkanCommandListContext*>& Contexts = CommandContexts[InContext->GetPipeline()];
-	Contexts.Push(InContext);
-}
-
-FVulkanCommandListContext* FVulkanDevice::AcquireParallelRenderContext(FVulkanParallelRenderPassInfo* ParallelRenderPassInfo)
-{
-	FScopeLock Lock(&ContextCS);
-	FVulkanCommandListContext* ParallelRenderContext = nullptr;
-	if (ParallelRenderContexts.Num() > 0)
-	{
-		ParallelRenderContext = ParallelRenderContexts.Pop(EAllowShrinking::No);
-		ParallelRenderContext->SetParallelRenderPassInfo(ParallelRenderPassInfo);
-	}
-	else
-	{
-		ParallelRenderContext = new FVulkanCommandListContext(*this, ImmediateContext, ParallelRenderPassInfo);
-	}
-	return ParallelRenderContext;
-}
-
-void FVulkanDevice::ReleaseParallelRenderContext(FVulkanCommandListContext* InContext)
-{
-	check(InContext);
-	FScopeLock Lock(&ContextCS);
-	InContext->SetParallelRenderPassInfo(nullptr);
-	ParallelRenderContexts.Push(InContext);
-}
-
 void FVulkanDevice::VulkanSetObjectName(VkObjectType Type, uint64_t Handle, const TCHAR* Name)
 {
 #if VULKAN_ENABLE_DRAW_MARKERS
diff --git a/Engine/Source/Runtime/VulkanRHI/Private/VulkanDevice.h b/Engine/Source/Runtime/VulkanRHI/Private/VulkanDevice.h
--- a/Engine/Source/Runtime/VulkanRHI/Private/VulkanDevice.h
+++ b/Engine/Source/Runtime/VulkanRHI/Private/VulkanDevice.h
@@ -510,12 +510,6 @@ public:
 	void NotifyDeletedGfxPipeline(class FVulkanRHIGraphicsPipelineState* Pipeline);
 	void NotifyDeletedComputePipeline(class FVulkanComputePipeline* Pipeline);
 
-	FVulkanCommandListContext* AcquireDeferredContext(ERHIPipeline Pipeline);
-	void ReleaseDeferredContext(FVulkanCommandListContext* InContext);
-
-	FVulkanCommandListContext* AcquireParallelRenderContext(FVulkanParallelRenderPassInfo* ParallelRenderPassInfo);
-	void ReleaseParallelRenderContext(FVulkanCommandListContext* InContext);
-
 	void VulkanSetObjectName(VkObjectType Type, uint64_t Handle, const TCHAR* Name);
 	const FOptionalVulkanDeviceExtensions& GetOptionalExtensions() const
 	{
@@ -667,10 +661,7 @@ private:
 
 	TMap<uint32, FSamplerStateRHIRef> SamplerMap;
 
-	FCriticalSection ContextCS;
 	FVulkanCommandListContextImmediate* ImmediateContext;
-	TRHIPipelineArray<TArray<FVulkanCommandListContext*>> CommandContexts;
-	TArray<FVulkanCommandListContext*> ParallelRenderContexts;
 
 	FVulkanDynamicRHI* RHI = nullptr;
 	
diff --git a/Engine/Source/Runtime/VulkanRHI/Private/VulkanQueue.cpp b/Engine/Source/Runtime/VulkanRHI/Private/VulkanQueue.cpp
--- a/Engine/Source/Runtime/VulkanRHI/Private/VulkanQueue.cpp
+++ b/Engine/Source/Runtime/VulkanRHI/Private/VulkanQueue.cpp
@@ -8,6 +8,7 @@
 #include "VulkanDevice.h"
 #include "VulkanMemory.h"
 #include "VulkanContext.h"
+#include "VulkanCommandBuffer.h"
 
 int32 GWaitForIdleOnSubmit = 0;
 FAutoConsoleVariableRef CVarVulkanWaitForIdleOnSubmit(
@@ -67,6 +68,15 @@ FVulkanQueue::~FVulkanQueue()
 	{
 		TimelineSempahore->Release();
 	}
+
+	for (TArray<FVulkanCommandBufferPool*>& CommandBufferPoolArray : CommandBufferPools)
+	{
+		for (FVulkanCommandBufferPool* CommandBufferPool : CommandBufferPoolArray)
+		{
+			delete CommandBufferPool;
+		}
+		CommandBufferPoolArray.Reset();
+	}
 }
 
 int32 FVulkanQueue::SubmitQueuedPayloads(TMap<VkSemaphore, FBinarySemaphoreSignalInfo>& SignaledSemas)
@@ -347,6 +357,8 @@ void FVulkanQueue::Submit(TArrayView<VkSubmitInfo> InSubmitInfos, VulkanRHI::FFe
 
 int32 FVulkanQueue::ProcessInterruptQueue(uint64 Timeout)
 {
+	SCOPED_NAMED_EVENT_TEXT("ProcessInterruptQueue", FColor::Orange);
+
 	if (bUseTimelineSemaphores)
 	{
 		checkSlow(TimelineSempahore);
@@ -521,4 +533,23 @@ UE::RHI::GPUProfiler::FQueue FVulkanQueue::GetProfilerQueue() const
 
 	return ProfilerQueue;
 }
-#endif // RHI_NEW_GPU_PROFILER
\ No newline at end of file
+#endif // RHI_NEW_GPU_PROFILER
+
+FVulkanCommandBufferPool* FVulkanQueue::AcquireCommandBufferPool(EVulkanCommandBufferType CommandBufferType)
+{
+	FScopeLock Lock(&CommandBufferPoolCS);
+	TArray<FVulkanCommandBufferPool*>& CommandBufferPoolArray = CommandBufferPools[(int32)CommandBufferType];
+	if (CommandBufferPoolArray.Num())
+	{
+		return CommandBufferPoolArray.Pop(EAllowShrinking::No);
+	}
+	return new FVulkanCommandBufferPool(Device, *this, CommandBufferType);
+}
+
+void FVulkanQueue::ReleaseCommandBufferPool(FVulkanCommandBufferPool* CommandBufferPool)
+{
+	FScopeLock Lock(&CommandBufferPoolCS);
+	check(&CommandBufferPool->GetQueue() == this);
+	TArray<FVulkanCommandBufferPool*>& CommandBufferPoolArray = CommandBufferPools[(int32)CommandBufferPool->GetCommandBufferType()];
+	CommandBufferPoolArray.Add(CommandBufferPool);
+}
diff --git a/Engine/Source/Runtime/VulkanRHI/Private/VulkanQueue.h b/Engine/Source/Runtime/VulkanRHI/Private/VulkanQueue.h
--- a/Engine/Source/Runtime/VulkanRHI/Private/VulkanQueue.h
+++ b/Engine/Source/Runtime/VulkanRHI/Private/VulkanQueue.h
@@ -8,6 +8,7 @@
 
 #include "VulkanConfiguration.h"
 #include "VulkanBarriers.h"
+#include "VulkanCommandBuffer.h"
 
 class FVulkanDynamicRHI;
 class FVulkanDevice;
@@ -61,6 +62,9 @@ public:
 		return Queue;
 	}
 
+	FVulkanCommandBufferPool* AcquireCommandBufferPool(EVulkanCommandBufferType CommandBufferType);
+	void ReleaseCommandBufferPool(FVulkanCommandBufferPool* CommandBufferPool);
+
 	VkPipelineStageFlags GetSupportedStageBits() const
 	{
 		return SupportedStages;
@@ -117,12 +121,14 @@ private:
 	EVulkanQueueType QueueType;
 	FVulkanDevice& Device;
 
+	FCriticalSection CommandBufferPoolCS;
+	TStaticArray<TArray<FVulkanCommandBufferPool*>, (int32)EVulkanCommandBufferType::Count> CommandBufferPools;
+
 	const bool bUseTimelineSemaphores;
 	VulkanRHI::FSemaphore* TimelineSempahore = nullptr;
 	uint64 NextTimelineSemaphoreValue = 1;
 	uint64 CompletedTimelineSemaphoreValue = 0;
 
-	mutable FCriticalSection CS;
 	uint64 SubmitCounter;
 	VkPipelineStageFlags SupportedStages = VK_PIPELINE_STAGE_NONE;
 	VkAccessFlags SupportedAccess = VK_ACCESS_NONE;
diff --git a/Engine/Source/Runtime/VulkanRHI/Private/VulkanSubmission.cpp b/Engine/Source/Runtime/VulkanRHI/Private/VulkanSubmission.cpp
--- a/Engine/Source/Runtime/VulkanRHI/Private/VulkanSubmission.cpp
+++ b/Engine/Source/Runtime/VulkanRHI/Private/VulkanSubmission.cpp
@@ -239,7 +239,7 @@ void FVulkanDynamicRHI::RHIFinalizeContext(FRHIFinalizeContextArgs&& Args, TRHIP
 		}
 		else
 		{
-			Device->ReleaseDeferredContext(VulkanContext);
+			delete VulkanContext;
 		}
 	}
 }
@@ -248,7 +248,7 @@ IRHIComputeContext* FVulkanDynamicRHI::RHIGetParallelCommandContext(FRHIParallel
 {
 	FVulkanParallelRenderPassInfo* VulkanParallelRenderPassInfo = (FVulkanParallelRenderPassInfo*)ParallelRenderPass.RHIPlatformData;
 	checkf(VulkanParallelRenderPassInfo, TEXT("Must begin parallel render pass on parent context before creating parallel contexts!"));
-	return Device->AcquireParallelRenderContext(VulkanParallelRenderPassInfo);
+	return new FVulkanCommandListContext(*Device, &Device->GetImmediateContext(), VulkanParallelRenderPassInfo);
 }
 
 static FCriticalSection GSecondaryPayloadsCS;
@@ -260,10 +260,11 @@ IRHIPlatformCommandList* FVulkanDynamicRHI::RHIFinalizeParallelContext(IRHICompu
 	{
 		FScopeLock Lock(&GSecondaryPayloadsCS);
 		FVulkanParallelRenderPassInfo* VulkanParallelRenderPassInfo = VulkanContext->GetParallelRenderPassInfo();
+		checkf(VulkanParallelRenderPassInfo, TEXT("Attempting to call RHIFinalizeParallelContext on a non-parallel context!"));
 		VulkanContext->Finalize(VulkanParallelRenderPassInfo->SecondaryPayloads);
 	}
 
-	Device->ReleaseParallelRenderContext(VulkanContext);
+	delete VulkanContext;
 	return nullptr;
 }
 
@@ -321,6 +322,8 @@ bool FVulkanDynamicRHI::ProcessSubmissionQueue()
 
 bool FVulkanDynamicRHI::WaitAndProcessInterruptQueue()
 {
+	SCOPED_NAMED_EVENT_TEXT("WaitAndProcessInterruptQueue", FColor::Orange);
+
 	bool bProgress = false;
 
 	// Wait for progress
