VulkanRHI: Added support for RunOnQueue function in VulkanRHI

#rnx

diff --git a/Engine/Source/Runtime/VulkanRHI/Private/VulkanDynamicRHI.h b/Engine/Source/Runtime/VulkanRHI/Private/VulkanDynamicRHI.h
--- a/Engine/Source/Runtime/VulkanRHI/Private/VulkanDynamicRHI.h
+++ b/Engine/Source/Runtime/VulkanRHI/Private/VulkanDynamicRHI.h
@@ -324,6 +324,9 @@ public:
 		return ApiVersion;
 	}
 
+	// Runs code on SubmissionThread with access to VkQueue.  Useful for plugins. 
+	virtual void RHIRunOnQueue(EVulkanRHIRunOnQueueType QueueType, TFunction<void(VkQueue)>&& CodeToRun, bool bWaitForSubmission) final override;
+
 public:
 	static void SavePipelineCache();
 
diff --git a/Engine/Source/Runtime/VulkanRHI/Private/VulkanQueue.cpp b/Engine/Source/Runtime/VulkanRHI/Private/VulkanQueue.cpp
--- a/Engine/Source/Runtime/VulkanRHI/Private/VulkanQueue.cpp
+++ b/Engine/Source/Runtime/VulkanRHI/Private/VulkanQueue.cpp
@@ -197,6 +197,8 @@ void FVulkanQueue::SubmitPayloads(TArrayView<FVulkanPayload*> Payloads, TMap<VkS
 
 	for (FVulkanPayload* Payload : Payloads)
 	{
+		Payload->PreExecute();
+
 #if RHI_NEW_GPU_PROFILER
 		if (Payload->EndFrameEvent.IsSet())
 		{
diff --git a/Engine/Source/Runtime/VulkanRHI/Private/VulkanSubmission.cpp b/Engine/Source/Runtime/VulkanRHI/Private/VulkanSubmission.cpp
--- a/Engine/Source/Runtime/VulkanRHI/Private/VulkanSubmission.cpp
+++ b/Engine/Source/Runtime/VulkanRHI/Private/VulkanSubmission.cpp
@@ -559,6 +559,7 @@ void FVulkanDynamicRHI::RHIBlockUntilGPUIdle()
 	}
 }
 
+
 // Resolve any pending actions and delete the payload
 void FVulkanDynamicRHI::CompletePayload(FVulkanPayload* Payload)
 {
@@ -677,3 +678,53 @@ void FVulkanDynamicRHI::CompletePayload(FVulkanPayload* Payload)
 
 	delete Payload;
 }
+
+void FVulkanPayload::PreExecute()
+{
+	if (PreExecuteCallback)
+	{
+		PreExecuteCallback(Queue.GetHandle());
+	}
+}
+
+void FVulkanDynamicRHI::RHIRunOnQueue(EVulkanRHIRunOnQueueType QueueType, TFunction<void(VkQueue)>&& CodeToRun, bool bWaitForSubmission)
+{
+	FGraphEventRef SubmissionEvent;
+
+	FVulkanQueue* Queue = nullptr;
+	switch (QueueType)
+	{
+	case EVulkanRHIRunOnQueueType::Graphics:
+		Queue = Device->GetQueue(EVulkanQueueType::Graphics);
+		break;
+	case EVulkanRHIRunOnQueueType::Transfer:
+		Queue = Device->GetQueue(EVulkanQueueType::Transfer);
+		break;
+	default:
+		checkf(false, TEXT("Unknown EVulkanRHIRunOnQueueType, skipping call."));
+		return;
+	}
+
+	FVulkanPlatformCommandList* Payloads = new FVulkanPlatformCommandList;
+	FVulkanPayload* Payload = new FVulkanPayload(*Queue);
+	Payloads->Add(Payload);
+
+	Payload->PreExecuteCallback = MoveTemp(CodeToRun);
+
+	if (bWaitForSubmission)
+	{
+		SubmissionEvent = FGraphEvent::CreateGraphEvent();
+		Payload->SubmissionEvents.Add(SubmissionEvent);
+	}
+
+	PendingPayloadsForSubmission.Enqueue(Payloads);
+	KickSubmissionThread();
+
+	// Use this opportunity to pump the interrupt queue
+	ProcessInterruptQueueUntil(nullptr);
+
+	if (SubmissionEvent && !SubmissionEvent->IsComplete())
+	{
+		SubmissionEvent->Wait();
+	}
+}
diff --git a/Engine/Source/Runtime/VulkanRHI/Private/VulkanSubmission.h b/Engine/Source/Runtime/VulkanRHI/Private/VulkanSubmission.h
--- a/Engine/Source/Runtime/VulkanRHI/Private/VulkanSubmission.h
+++ b/Engine/Source/Runtime/VulkanRHI/Private/VulkanSubmission.h
@@ -45,6 +45,8 @@ public:
 
 protected:
 
+	void PreExecute();
+
 	FVulkanQueue& Queue;
 	TArray<VkPipelineStageFlags> WaitFlags;  // flags that match 1:1 with WaitSemaphores
 	TArray<VulkanRHI::FSemaphore*> WaitSemaphores; // wait before command buffers
@@ -64,6 +66,9 @@ protected:
 	// Queries used in the command lists of this payload
 	TStaticArray<TArray<FVulkanQueryPool*>, (int32)EVulkanQueryPoolType::Count> QueryPools;
 
+	// Used by RHIRunOnQueue
+	TFunction<void(VkQueue)> PreExecuteCallback;
+
 #if WITH_RHI_BREADCRUMBS
 	FRHIBreadcrumbRange BreadcrumbRange{};
 	TSharedPtr<FRHIBreadcrumbAllocatorArray> BreadcrumbAllocators{};
diff --git a/Engine/Source/Runtime/VulkanRHI/Public/IVulkanDynamicRHI.h b/Engine/Source/Runtime/VulkanRHI/Public/IVulkanDynamicRHI.h
--- a/Engine/Source/Runtime/VulkanRHI/Public/IVulkanDynamicRHI.h
+++ b/Engine/Source/Runtime/VulkanRHI/Public/IVulkanDynamicRHI.h
@@ -30,6 +30,12 @@ struct FVulkanRHIExternalImageDeleteCallbackInfo
 	void (*Function)(void* UserData) = nullptr;
 };
 
+enum class EVulkanRHIRunOnQueueType
+{
+	Graphics = 0,
+	Transfer,
+};
+
 struct IVulkanDynamicRHI : public FDynamicRHI
 {
 	virtual ERHIInterfaceType GetInterfaceType() const override { return ERHIInterfaceType::Vulkan; }
@@ -91,6 +97,9 @@ struct IVulkanDynamicRHI : public FDynamicRHI
 
 	static VULKANRHI_API void AddEnabledInstanceExtensionsAndLayers(TArrayView<const ANSICHAR* const> InInstanceExtensions, TArrayView<const ANSICHAR* const> InInstanceLayers);
 	static VULKANRHI_API void AddEnabledDeviceExtensionsAndLayers(TArrayView<const ANSICHAR* const> InDeviceExtensions, TArrayView<const ANSICHAR* const> InDeviceLayers);
+
+	// Runs code on SubmissionThread with access to VkQueue.  Useful for plugins. 
+	virtual void RHIRunOnQueue(EVulkanRHIRunOnQueueType QueueType, TFunction<void(VkQueue)>&& CodeToRun, bool bWaitForSubmission) = 0;
 };
 
 inline IVulkanDynamicRHI* GetIVulkanDynamicRHI()
