VulkanRHI: Try to read loader version from filesystem and default to using the included SDK loader for the version known to have issues with ASAN.

#rnx

diff --git a/Engine/Source/Runtime/VulkanRHI/Private/Linux/VulkanLinuxPlatform.cpp b/Engine/Source/Runtime/VulkanRHI/Private/Linux/VulkanLinuxPlatform.cpp
--- a/Engine/Source/Runtime/VulkanRHI/Private/Linux/VulkanLinuxPlatform.cpp
+++ b/Engine/Source/Runtime/VulkanRHI/Private/Linux/VulkanLinuxPlatform.cpp
@@ -43,6 +43,70 @@ bool FVulkanLinuxPlatform::IsSupported()
 	return LoadVulkanLibrary();
 }
 
+// vkEnumerateInstanceVersion triggers an ASAN errors at some versions, use the filesystem to determine the version
+// Only reject the loader if we were able to confirm its version number, for this reason we return MAX_uint32 on error
+static uint32_t GetVulkanInstanceVersion(void* VulkanLoader, const ANSICHAR* LoaderFilename)
+{
+	FAnsiString FullPath;
+	{
+		ANSICHAR LoaderPath[PATH_MAX+1];
+		FMemory::Memzero(LoaderPath);
+		const int32 RetVal = dlinfo(VulkanLoader, RTLD_DI_ORIGIN, LoaderPath);
+		if (RetVal < 0)
+		{
+			return MAX_uint32;
+		}
+		UE_LOG(LogVulkanRHI, Display, TEXT("Installed Vulkan Loader Path: %s"), ANSI_TO_TCHAR(LoaderPath));
+		FullPath = LoaderPath;
+	}
+	FullPath += "/";
+	FullPath += LoaderFilename;
+	
+	auto IsSymLink = [](const ANSICHAR* Path)
+	{
+		struct stat PathStat;
+		if (lstat(Path, &PathStat) == -1)
+		{
+			return false;
+		}
+		
+		return S_ISLNK(PathStat.st_mode);
+	};
+
+	auto FollowLinkTarget = [](FAnsiString& LinkPath)
+	{
+		ANSICHAR TargetPath[PATH_MAX+1];
+		FMemory::Memzero(TargetPath);
+		if (readlink(*LinkPath, TargetPath, PATH_MAX) > 0)
+		{
+			LinkPath = TargetPath;
+			return true;
+		}
+		return false;
+	};
+	
+		
+	while (IsSymLink(*FullPath))
+	{
+		FollowLinkTarget(FullPath);
+	}
+	
+	TArray<FAnsiString> SplitStrings;
+	FullPath.ParseIntoArray(SplitStrings, ".", true);
+	if (SplitStrings.Num() >= 4 && 
+		SplitStrings[SplitStrings.Num()-3].IsNumeric() &&
+		SplitStrings[SplitStrings.Num()-2].IsNumeric() &&
+		SplitStrings[SplitStrings.Num()-1].IsNumeric())
+	{
+		return VK_MAKE_API_VERSION(0, 
+			FCStringAnsi::Atoi(*SplitStrings[SplitStrings.Num()-3]), 
+			FCStringAnsi::Atoi(*SplitStrings[SplitStrings.Num()-2]), 
+			FCStringAnsi::Atoi(*SplitStrings[SplitStrings.Num()-1]));
+	}
+	
+	return MAX_uint32;
+}
+
 bool FVulkanLinuxPlatform::LoadVulkanLibrary()
 {
 	if (bAttemptedLoad)
@@ -87,19 +151,47 @@ bool FVulkanLinuxPlatform::LoadVulkanLibrary()
 	UE_LOG(LogVulkanRHI, Display, TEXT("Updated LD_LIBRARY_PATH=%s"), *PreviousLibPath);
 #endif // VULKAN_HAS_DEBUGGING_ENABLED
 
+	if (!FParse::Param(FCommandLine::Get(), TEXT("UseLocalVulkanLoader")))
+	{
+		const ANSICHAR* LoaderFilename = "libvulkan.so.1";
+		
+		// Try to load the default libvulkan.so
+		void* GlobalVulkanLib = dlopen(LoaderFilename, RTLD_NOW | RTLD_LOCAL);
+		
+		if (FParse::Param(FCommandLine::Get(), TEXT("UseGlobalVulkanLoader")))
+		{
+			VulkanLib = GlobalVulkanLib;
+			GlobalVulkanLib = nullptr;
+		}
+		else if (GlobalVulkanLib)
+		{
+			// Verify version and discard it if the version has issues
+			const uint32_t ApiVersion = GetVulkanInstanceVersion(GlobalVulkanLib, LoaderFilename);
+			UE_LOG(LogVulkanRHI, Display, TEXT("Installed Vulkan Loader instance version %u.%u.%u."), 
+				VK_API_VERSION_MAJOR(ApiVersion), VK_API_VERSION_MINOR(ApiVersion), VK_API_VERSION_PATCH(ApiVersion));
+						
+			if (ApiVersion > VK_MAKE_API_VERSION(0, 1, 3, 204))
+			{
+				VulkanLib = GlobalVulkanLib;
+				GlobalVulkanLib = nullptr;
+			}
+		}
+
+		if (GlobalVulkanLib)
+		{
+			dlclose(GlobalVulkanLib);
+		}
+	}
+
 	// Try to load libvulkan.so from the included SDK
-	if (FParse::Param(FCommandLine::Get(), TEXT("UseLocalVulkanLoader")))
+	if (VulkanLib == nullptr)
 	{
+		UE_LOG(LogVulkanRHI, Display, TEXT("Using included Vulkan loader."));
 		FAnsiString VulkanLoaderPath(UEVulkanBinariesPath);
 		VulkanLoaderPath.Append("/libvulkan.so");
 		VulkanLib = dlopen(*VulkanLoaderPath, RTLD_NOW | RTLD_LOCAL);
 	}
 
-	if (VulkanLib == nullptr)
-	{
-		// Try to load the default libvulkan.so
-		VulkanLib = dlopen("libvulkan.so.1", RTLD_NOW | RTLD_LOCAL);
-	}
 
 	if (VulkanLib == nullptr)
 	{
@@ -110,7 +202,7 @@ bool FVulkanLinuxPlatform::LoadVulkanLibrary()
 
 		return false;
 	}
-
+	
 	bool bFoundAllEntryPoints = true;
 #define CHECK_VK_ENTRYPOINTS(Type,Func) if (VulkanDynamicAPI::Func == NULL) { bFoundAllEntryPoints = false; UE_LOG(LogRHI, Warning, TEXT("Failed to find entry point for %s"), TEXT(#Func)); }
 
