VulkanRHI: Set secondary command buffers as "Submitted" when the parent is submitted (fix leak).  Don't add BeginWork/EndWork timestamps on secondary command buffers since they execute within the Begin/End of the parent.

#rnx

diff --git a/Engine/Source/Runtime/VulkanRHI/Private/VulkanCommandBuffer.cpp b/Engine/Source/Runtime/VulkanRHI/Private/VulkanCommandBuffer.cpp
--- a/Engine/Source/Runtime/VulkanRHI/Private/VulkanCommandBuffer.cpp
+++ b/Engine/Source/Runtime/VulkanRHI/Private/VulkanCommandBuffer.cpp
@@ -287,6 +287,11 @@ void FVulkanCommandBuffer::Reset()
 
 void FVulkanCommandBuffer::SetSubmitted()
 {
+	for (FVulkanCommandBuffer* SecondaryCommand : ExecutedSecondaryCommandBuffers)
+	{
+		SecondaryCommand->SetSubmitted();
+	}
+
 	FScopeLock Lock(CommandBufferPool.GetCS());
 	State = FVulkanCommandBuffer::EState::Submitted;
 	SubmittedTime = FPlatformTime::Seconds();
diff --git a/Engine/Source/Runtime/VulkanRHI/Private/VulkanContext.cpp b/Engine/Source/Runtime/VulkanRHI/Private/VulkanContext.cpp
--- a/Engine/Source/Runtime/VulkanRHI/Private/VulkanContext.cpp
+++ b/Engine/Source/Runtime/VulkanRHI/Private/VulkanContext.cpp
@@ -43,13 +43,18 @@ void FVulkanContextCommon::EndPayload()
 			FVulkanCommandBuffer* CommandBuffer = Payload->CommandBuffers.Last();
 			checkSlow(!CommandBuffer->IsSubmitted() && CommandBuffer->HasBegun());
 
-			if (!CommandBuffer->IsOutsideRenderPass())
+			const bool bIsPrimaryCommandBuffer = (CommandBuffer->GetCommandBufferType() == EVulkanCommandBufferType::Primary);
+			if (!CommandBuffer->IsOutsideRenderPass() && bIsPrimaryCommandBuffer)
 			{
 				UE_LOG(LogVulkanRHI, Warning, TEXT("Forcing EndRenderPass() for submission"));
 				CommandBuffer->EndRenderPass();
 			}
 
-			CommandBuffer->End(GRHIGlobals.SupportsTimestampRenderQueries ? GetCurrentTimestampQueryPool(*Payload) : nullptr);
+			// Only record begin/end timestamps on primary command buffers
+			FVulkanQueryPool* TimestampQueryPool = 
+				GRHIGlobals.SupportsTimestampRenderQueries && bIsPrimaryCommandBuffer ?
+				GetCurrentTimestampQueryPool(*Payload) : nullptr;
+			CommandBuffer->End(TimestampQueryPool);
 		}
 	}
 }
@@ -170,8 +175,11 @@ void FVulkanContextCommon::PrepareNewCommandBuffer(FVulkanPayload& Payload)
 
 	Payload.CommandBuffers.Add(NewCommandBuffer);
 
-	FVulkanQueryPool* TimestampQueryPool = GRHIGlobals.SupportsTimestampRenderQueries ? GetCurrentTimestampQueryPool(Payload) : nullptr;
-	const VkRenderPass RenderPassHandle = GetParallelRenderPassInfo() ? GetParallelRenderPassInfo()->RenderPassHandle : VK_NULL_HANDLE;
+	// Only record begin/end timestamps on primary command buffers
+	const bool bIsPrimaryCommandBuffer = (NewCommandBuffer->GetCommandBufferType() == EVulkanCommandBufferType::Primary);
+	FVulkanQueryPool* TimestampQueryPool = GRHIGlobals.SupportsTimestampRenderQueries && bIsPrimaryCommandBuffer ?
+		GetCurrentTimestampQueryPool(Payload) : nullptr;
+	VkRenderPass RenderPassHandle = GetParallelRenderPassInfo() ? GetParallelRenderPassInfo()->RenderPassHandle : VK_NULL_HANDLE;
 	NewCommandBuffer->Begin(TimestampQueryPool, RenderPassHandle);
 }
 
@@ -255,6 +263,7 @@ void FVulkanCommandListContext::RHIEndParallelRenderPass()
 {
 	if (CurrentParallelRenderPassInfo->SecondaryPayloads.Num())
 	{
+		FVulkanPayload& ParentPayload = GetPayload(EPhase::Execute);
 		FVulkanCommandBuffer& ParentCommandBuffer = GetCommandBuffer();
 
 		TArray<VkCommandBuffer> CommandBufferHandles;
@@ -264,9 +273,20 @@ void FVulkanCommandListContext::RHIEndParallelRenderPass()
 			check(Payload->SignalSemaphores.Num() == 0);
 			check(Payload->WaitSemaphores.Num() == 0);
 
+			for (int32 QueryPoolTypeIndex = 0; QueryPoolTypeIndex < (int32)EVulkanQueryPoolType::Count; ++QueryPoolTypeIndex)
+			{
+				ParentPayload.QueryPools[QueryPoolTypeIndex].Append(MoveTemp(Payload->QueryPools[QueryPoolTypeIndex]));
+			}
+
 			for (FVulkanCommandBuffer* SecondaryCommandBuffer : Payload->CommandBuffers)
 			{
 				CommandBufferHandles.Add(SecondaryCommandBuffer->GetHandle());
+
+#if RHI_NEW_GPU_PROFILER
+				// :todo: Flush the event stream if there is one (only use the parent timings for now)
+				UE::RHI::GPUProfiler::FEventStream Dummy(Device.GetQueue(EVulkanQueueType::Graphics)->GetProfilerQueue());
+				SecondaryCommandBuffer->FlushProfilerEvents(Dummy, FPlatformTime::Cycles64());
+#endif
 			}
 			ParentCommandBuffer.ExecutedSecondaryCommandBuffers.Append(MoveTemp(Payload->CommandBuffers));
 
diff --git a/Engine/Source/Runtime/VulkanRHI/Private/VulkanContext.h b/Engine/Source/Runtime/VulkanRHI/Private/VulkanContext.h
--- a/Engine/Source/Runtime/VulkanRHI/Private/VulkanContext.h
+++ b/Engine/Source/Runtime/VulkanRHI/Private/VulkanContext.h
@@ -41,6 +41,7 @@ struct FVulkanParallelRenderPassInfo
 
 class FVulkanContextCommon
 {
+protected:
 	// Phases are used to track where we are at in filling the current payload.
 	// Phases always move forward in the order declared here, never backwards.
 	// Always use GetPayload(EPhase) to fetch the payload for the phase you want,
diff --git a/Engine/Source/Runtime/VulkanRHI/Private/VulkanRHI.cpp b/Engine/Source/Runtime/VulkanRHI/Private/VulkanRHI.cpp
--- a/Engine/Source/Runtime/VulkanRHI/Private/VulkanRHI.cpp
+++ b/Engine/Source/Runtime/VulkanRHI/Private/VulkanRHI.cpp
@@ -1297,7 +1297,7 @@ void FVulkanCommandListContext::RHIEndDrawingViewport(FRHIViewport* ViewportRHI,
 		{
 			FlushProfilerStats();
 
-			if (GRHIGlobals.SupportsTimestampRenderQueries)
+			if (GRHIGlobals.SupportsTimestampRenderQueries && !CurrentParallelRenderPassInfo)
 			{
 				auto& Event = GetCommandBuffer().EmplaceProfilerEvent<UE::RHI::GPUProfiler::FEvent::FBeginBreadcrumb>(Breadcrumb);
 				FVulkanQueryPool* CurrentPool = GetCurrentTimestampQueryPool();
@@ -1328,7 +1328,7 @@ void FVulkanCommandListContext::RHIEndDrawingViewport(FRHIViewport* ViewportRHI,
 		{
 			FlushProfilerStats();
 
-			if (GRHIGlobals.SupportsTimestampRenderQueries)
+			if (GRHIGlobals.SupportsTimestampRenderQueries && !CurrentParallelRenderPassInfo)
 			{
 				auto& Event = GetCommandBuffer().EmplaceProfilerEvent<UE::RHI::GPUProfiler::FEvent::FEndBreadcrumb>(Breadcrumb);
 				FVulkanQueryPool* CurrentPool = GetCurrentTimestampQueryPool();
