VulkanRHI: Do not emit stats if timestamps aren't supported.

#rnx

diff --git a/Engine/Source/Runtime/VulkanRHI/Private/VulkanRHI.cpp b/Engine/Source/Runtime/VulkanRHI/Private/VulkanRHI.cpp
--- a/Engine/Source/Runtime/VulkanRHI/Private/VulkanRHI.cpp
+++ b/Engine/Source/Runtime/VulkanRHI/Private/VulkanRHI.cpp
@@ -1294,16 +1294,14 @@ void FVulkanCommandListContext::RHIEndDrawingViewport(FRHIViewport* ViewportRHI,
 		}
 		
 #if RHI_NEW_GPU_PROFILER
+		if (GRHIGlobals.SupportsTimestampRenderQueries && !CurrentParallelRenderPassInfo)
 		{
 			FlushProfilerStats();
 
-			if (GRHIGlobals.SupportsTimestampRenderQueries && !CurrentParallelRenderPassInfo)
-			{
-				auto& Event = GetCommandBuffer().EmplaceProfilerEvent<UE::RHI::GPUProfiler::FEvent::FBeginBreadcrumb>(Breadcrumb);
-				FVulkanQueryPool* CurrentPool = GetCurrentTimestampQueryPool();
-				const uint32 IndexInPool = CurrentPool->ReserveQuery(&Event.GPUTimestampTOP);
-				VulkanRHI::vkCmdWriteTimestamp(GetCommandBuffer().GetHandle(), VK_PIPELINE_STAGE_TOP_OF_PIPE_BIT, CurrentPool->GetHandle(), IndexInPool);
-			}
+			auto& Event = GetCommandBuffer().EmplaceProfilerEvent<UE::RHI::GPUProfiler::FEvent::FBeginBreadcrumb>(Breadcrumb);
+			FVulkanQueryPool* CurrentPool = GetCurrentTimestampQueryPool();
+			const uint32 IndexInPool = CurrentPool->ReserveQuery(&Event.GPUTimestampTOP);
+			VulkanRHI::vkCmdWriteTimestamp(GetCommandBuffer().GetHandle(), VK_PIPELINE_STAGE_TOP_OF_PIPE_BIT, CurrentPool->GetHandle(), IndexInPool);
 		}
 #else
 		if (IsImmediate())
@@ -1325,16 +1323,14 @@ void FVulkanCommandListContext::RHIEndDrawingViewport(FRHIViewport* ViewportRHI,
 	void FVulkanCommandListContext::RHIEndBreadcrumbGPU(FRHIBreadcrumbNode* Breadcrumb)
 	{
 #if RHI_NEW_GPU_PROFILER
+		if (GRHIGlobals.SupportsTimestampRenderQueries && !CurrentParallelRenderPassInfo)
 		{
 			FlushProfilerStats();
 
-			if (GRHIGlobals.SupportsTimestampRenderQueries && !CurrentParallelRenderPassInfo)
-			{
-				auto& Event = GetCommandBuffer().EmplaceProfilerEvent<UE::RHI::GPUProfiler::FEvent::FEndBreadcrumb>(Breadcrumb);
-				FVulkanQueryPool* CurrentPool = GetCurrentTimestampQueryPool();
-				const uint32 IndexInPool = CurrentPool->ReserveQuery(&Event.GPUTimestampBOP);
-				VulkanRHI::vkCmdWriteTimestamp(GetCommandBuffer().GetHandle(), VK_PIPELINE_STAGE_BOTTOM_OF_PIPE_BIT, CurrentPool->GetHandle(), IndexInPool);
-			}
+			auto& Event = GetCommandBuffer().EmplaceProfilerEvent<UE::RHI::GPUProfiler::FEvent::FEndBreadcrumb>(Breadcrumb);
+			FVulkanQueryPool* CurrentPool = GetCurrentTimestampQueryPool();
+			const uint32 IndexInPool = CurrentPool->ReserveQuery(&Event.GPUTimestampBOP);
+			VulkanRHI::vkCmdWriteTimestamp(GetCommandBuffer().GetHandle(), VK_PIPELINE_STAGE_BOTTOM_OF_PIPE_BIT, CurrentPool->GetHandle(), IndexInPool);
 		}
 #else
 		//only valid on immediate context currently.  needs to be fixed for parallel rhi execute
