VulkanRHI: Add missing call to FlushProfilerStats on Finalize()

#rnx

diff --git a/Engine/Source/Runtime/VulkanRHI/Private/VulkanContext.cpp b/Engine/Source/Runtime/VulkanRHI/Private/VulkanContext.cpp
--- a/Engine/Source/Runtime/VulkanRHI/Private/VulkanContext.cpp
+++ b/Engine/Source/Runtime/VulkanRHI/Private/VulkanContext.cpp
@@ -63,6 +63,8 @@ void FVulkanContextCommon::EndPayload()
 // payloads to the given array. Resets the context so new commands can be recorded.
 void FVulkanContextCommon::Finalize(TArray<FVulkanPayload*>& OutPayloads)
 {
+	FlushProfilerStats();
+
 	FlushPendingSyncPoints();
 
 	if (ContextSyncPoint.IsValid())
@@ -188,6 +190,7 @@ FVulkanCommandListContext::FVulkanCommandListContext(FVulkanDevice& InDevice, ER
 	: FVulkanContextCommon(InDevice, *InDevice.GetQueue(InPipeline), EVulkanCommandBufferType::Primary)
 	, Immediate(InImmediate)
 	, RHIPipeline(InPipeline)
+	, bSupportsBreadcrumbs(GRHIGlobals.SupportsTimestampRenderQueries)
 #if (RHI_NEW_GPU_PROFILER == 0)
 	, GpuProfiler(this, &InDevice)
 #endif
@@ -208,6 +211,7 @@ FVulkanCommandListContext::FVulkanCommandListContext(FVulkanDevice& InDevice, FV
 	: FVulkanContextCommon(InDevice, *InDevice.GetQueue(ERHIPipeline::Graphics), EVulkanCommandBufferType::Secondary)
 	, Immediate(InImmediate)
 	, RHIPipeline(ERHIPipeline::Graphics)
+	, bSupportsBreadcrumbs(false)
 	, CurrentParallelRenderPassInfo(InParallelRenderPassInfo)
 #if (RHI_NEW_GPU_PROFILER == 0)
 	, GpuProfiler(this, &InDevice)
diff --git a/Engine/Source/Runtime/VulkanRHI/Private/VulkanContext.h b/Engine/Source/Runtime/VulkanRHI/Private/VulkanContext.h
--- a/Engine/Source/Runtime/VulkanRHI/Private/VulkanContext.h
+++ b/Engine/Source/Runtime/VulkanRHI/Private/VulkanContext.h
@@ -196,6 +196,8 @@ public:
 		return nullptr;
 	}
 
+	virtual void FlushProfilerStats() { /* do nothing */ }
+
 private:
 	FVulkanQueryPool* GetCurrentTimestampQueryPool(FVulkanPayload& Payload);
 	void PrepareNewCommandBuffer(FVulkanPayload& Payload);
@@ -302,10 +304,10 @@ public:
 	virtual void RHIEndRenderQuery(FRHIRenderQuery* RenderQuery) final override;
 
 #if RHI_NEW_GPU_PROFILER
-	void FlushProfilerStats()
+	virtual void FlushProfilerStats() override
 	{
-		// Flush accumulated draw stats
-		if (StatEvent)
+		// Flush accumulated draw stats (if breadcrumbs are available to attach them to)
+		if (StatEvent && bSupportsBreadcrumbs)
 		{
 			GetCommandBuffer().EmplaceProfilerEvent<UE::RHI::GPUProfiler::FEvent::FStats>() = StatEvent;
 			StatEvent = {};
@@ -402,6 +404,7 @@ public:
 protected:
 	FVulkanCommandListContext* Immediate;
 	const ERHIPipeline RHIPipeline;
+	const bool bSupportsBreadcrumbs;
 
 	void BeginOcclusionQueryBatch(uint32 NumQueriesInBatch);
 
diff --git a/Engine/Source/Runtime/VulkanRHI/Private/VulkanRHI.cpp b/Engine/Source/Runtime/VulkanRHI/Private/VulkanRHI.cpp
--- a/Engine/Source/Runtime/VulkanRHI/Private/VulkanRHI.cpp
+++ b/Engine/Source/Runtime/VulkanRHI/Private/VulkanRHI.cpp
@@ -1294,7 +1294,7 @@ void FVulkanCommandListContext::RHIEndDrawingViewport(FRHIViewport* ViewportRHI,
 		}
 		
 #if RHI_NEW_GPU_PROFILER
-		if (GRHIGlobals.SupportsTimestampRenderQueries && !CurrentParallelRenderPassInfo)
+		if (bSupportsBreadcrumbs)
 		{
 			FlushProfilerStats();
 
@@ -1323,7 +1323,7 @@ void FVulkanCommandListContext::RHIEndDrawingViewport(FRHIViewport* ViewportRHI,
 	void FVulkanCommandListContext::RHIEndBreadcrumbGPU(FRHIBreadcrumbNode* Breadcrumb)
 	{
 #if RHI_NEW_GPU_PROFILER
-		if (GRHIGlobals.SupportsTimestampRenderQueries && !CurrentParallelRenderPassInfo)
+		if (bSupportsBreadcrumbs)
 		{
 			FlushProfilerStats();
 
