Refactor of command list state tracking to use a much simpler TrackedAccess member that is pushed automatically for basic RHICmdList.Transition methods down to the RHI thread. This takes advantage of the fact that RDG handles all complicated subresource transitions to avoid doing so at the command list level. This saves having to patch / deep copy transitions during dispatch and saves about ~100 bytes of tracking data on FRHIViewableResource. We already restrict RHICmdList.Transition from doing subresource transitions, so there is no additional constraint being enforced by this change.

The other constraint is that resource states have to be patched inside of RHI{Begin, End}Transitions. There is a new helper struct UE::RHICore::FResourceState that translates the FRHITransitionInfo into patched states. The downside is having to potentially translate states twice, but in practice the operation is only done one time for single pipe transitions.

Transient resources require a bit of care due to the restrictions D3D12 legacy barriers have on Async Compute with Render / Depth Stencil targets. The discard transition sets the previous state of the transient resource into the tracked access which is picked up when the resource is acquired again in the next frame. This works for 95% of cases, except for textures with multiple subresources where subresource states differ. If this case is encountered, all subresources are transitioned into a single state. This is communicated to the RHI by using ERHIAccess::Discard | IntermediateState as the final state, which tells the underlying RHI that its both a discard operation and to transition to an intermediate known state. This only applies when GRHIGlobals.NeedsTransientDiscardStateTracking is enabled. There is no measurably GPU perf impact to this as its very rare.

diff --git a/Engine/Source/Runtime/D3D12RHI/Private/D3D12Commands.cpp b/Engine/Source/Runtime/D3D12RHI/Private/D3D12Commands.cpp
--- a/Engine/Source/Runtime/D3D12RHI/Private/D3D12Commands.cpp
+++ b/Engine/Source/Runtime/D3D12RHI/Private/D3D12Commands.cpp
@@ -10,6 +10,7 @@ D3D12Commands.cpp: D3D RHI commands implementation.
 #include "D3D12ResourceCollection.h"
 #include "D3D12TextureReference.h"
 #include "RHIUniformBufferUtilities.h"
+#include "RHICoreTransitions.h"
 
 static int32 GD3D12AllowDiscardResources = 1;
 static FAutoConsoleVariableRef CVarD3D12AllowDiscardResources(
@@ -333,22 +334,6 @@ static bool ShouldProcessTransition(const FD3D12TransitionData* Data, bool bIsBe
 	return (bIsBeginTransition ? ProcessTransitionDuringBegin(Data) : !ProcessTransitionDuringBegin(Data));
 }
 
-static bool IsImpossibleToDiscardOnAsync(
-	FD3D12CommandContext& Context,
-	const FRHITransitionInfo& Info)
-{
-	const ED3D12QueueType QueueType = Context.QueueType;
-	if (QueueType == ED3D12QueueType::Async)
-	{
-		const auto [Resource, Texture] = GetResourceAndTexture(Context, Info);
-
-		// It is expected that this will return COMMON if the initial state isn't compatible with the queue type
-		const D3D12_RESOURCE_STATES InitialState = GetInitialResourceState(Resource->GetDesc(), QueueType);
-		return InitialState == D3D12_RESOURCE_STATE_COMMON;
-	}
-	return false;
-}
-
 struct FD3D12DiscardResource
 {
 	FD3D12DiscardResource(FD3D12Resource* InResource, EResourceTransitionFlags InFlags, uint32 InSubresource = D3D12_RESOURCE_BARRIER_ALL_SUBRESOURCES, const FD3D12Texture* InTexture = nullptr, const FD3D12RenderTargetView* InRTV = nullptr)
@@ -389,60 +374,69 @@ void FD3D12CommandContext::HandleReservedResourceCommits(const struct FD3D12Tran
 	}
 }
 
+bool IsImpossibleAsyncDiscardTransition(ERHIPipeline Pipeline, FRHITexture* Texture)
+{
+	return Pipeline == ERHIPipeline::AsyncCompute && Texture && EnumHasAnyFlags(Texture->GetDesc().Flags, ETextureCreateFlags::RenderTargetable | ETextureCreateFlags::DepthStencilTargetable);
+}
+
 void FD3D12CommandContext::HandleResourceDiscardTransitions(
 	const FD3D12TransitionData* TransitionData,
 	TArray<FD3D12DiscardResource>& ResourcesToDiscard)
 {
 	for (const FRHITransitionInfo& Info : TransitionData->TransitionInfos)
 	{
-		if ((Info.AccessBefore != ERHIAccess::Discard) || !EnumHasAnyFlags(Info.Flags, EResourceTransitionFlags::Clear | EResourceTransitionFlags::Discard))
-		{
-			continue;
-		}
+		const UE::RHICore::FResourceState ResourceState(*this, TransitionData->SrcPipelines, TransitionData->DstPipelines, Info);
 
-		if (IsImpossibleToDiscardOnAsync(*this, Info))
+		if (!EnumHasAnyFlags(ResourceState.AccessBefore, ERHIAccess::Discard))
 		{
 			continue;
 		}
 
 		ProcessResource(*this, Info, [&](const FRHITransitionInfo& Info, FD3D12Resource* Resource, FD3D12Texture* Texture, FD3D12Texture* DiscardTexture)
 		{
-			if (!Resource->RequiresResourceStateTracking())
+			D3D12_RESOURCE_STATES StateAfter  = GetInitialResourceState(Resource->GetDesc(), QueueType);
+			D3D12_RESOURCE_STATES StateBefore = StateAfter;
+
+			if (ResourceState.AccessBefore != ERHIAccess::Discard)
 			{
-				return;
+				StateBefore = GetD3D12ResourceState(ResourceState.AccessBefore & ~ERHIAccess::Discard, QueueType, Texture);
+			}
+
+			const bool bTransition = StateBefore != StateAfter;
+
+			if (bTransition)
+			{
+				// Transitions here should only occur on the Direct queue and when the prior Discard operation failed due to being on async compute.
+				ensure(IsImpossibleAsyncDiscardTransition(ResourceState.SrcPipelines, Texture) && QueueType == ED3D12QueueType::Direct);
 			}
 
-			D3D12_RESOURCE_STATES InitialState = GetInitialResourceState(Resource->GetDesc(), QueueType);
 			if (Info.IsWholeResource() || Resource->GetSubresourceCount() == 1)
 			{
-				if ((Info.DiscardAccessBefore != ERHIAccess::Unknown) && (Info.DiscardAccessBefore != ERHIAccess::Discard))
+				if (bTransition)
 				{
-					D3D12_RESOURCE_STATES StateBefore = GetD3D12ResourceState(Info.DiscardAccessBefore, QueueType, Texture);
-					TransitionResource(Resource, StateBefore, InitialState, D3D12_RESOURCE_BARRIER_ALL_SUBRESOURCES);
+					TransitionResource(Resource, StateBefore, StateAfter, D3D12_RESOURCE_BARRIER_ALL_SUBRESOURCES);
 				}
 				else
 				{
 					UpdateResidency(Resource);
 				}
 
-
 				const FD3D12RenderTargetView* RTV = nullptr;
 #if PLATFORM_REQUIRES_TYPELESS_RESOURCE_DISCARD_WORKAROUND
 				if (DiscardTexture)
 				{
 					RTV = DiscardTexture->GetRenderTargetView(0, -1);
 				}
-#endif // #if PLATFORM_REQUIRES_TYPELESS_RESOURCE_DISCARD_WORKAROUND				
+#endif // #if PLATFORM_REQUIRES_TYPELESS_RESOURCE_DISCARD_WORKAROUND
 				ResourcesToDiscard.Emplace(Resource, Info.Flags, D3D12_RESOURCE_BARRIER_ALL_SUBRESOURCES, Texture, RTV);
 			}
 			else
 			{
 				EnumerateSubresources(Resource, Info, DiscardTexture, [&](uint32 Subresource, const FD3D12RenderTargetView* RTV)
 				{
-					if ((Info.DiscardAccessBefore != ERHIAccess::Unknown) && (Info.DiscardAccessBefore != ERHIAccess::Discard))
+					if (bTransition)
 					{
-						D3D12_RESOURCE_STATES StateBefore = GetD3D12ResourceState(Info.DiscardAccessBefore, QueueType, Texture);
-						TransitionResource(Resource, StateBefore, InitialState, Subresource);
+						TransitionResource(Resource, StateBefore, StateAfter, Subresource);
 					}
 					else
 					{
@@ -484,37 +478,29 @@ void FD3D12CommandContext::HandleDiscardResources(
 
 	for (const FD3D12DiscardResource& DiscardResource : ResourcesToDiscard)
 	{
-		if (EnumHasAnyFlags(DiscardResource.Flags, EResourceTransitionFlags::Clear))
+#if PLATFORM_REQUIRES_TYPELESS_RESOURCE_DISCARD_WORKAROUND
+		if (DiscardResource.Texture && DiscardResource.RTV && DiscardResource.Texture->GetRequiresTypelessResourceDiscardWorkaround())
 		{
-			// add correct ops for clear
-			check(false);
+			FLinearColor ClearColor = DiscardResource.Texture->GetClearColor();
+			GraphicsCommandList()->ClearRenderTargetView(DiscardResource.RTV->GetOfflineCpuHandle(), reinterpret_cast<float*>(&ClearColor), 0, nullptr);
+			UpdateResidency(DiscardResource.RTV->GetResource());
 		}
-		else if (EnumHasAnyFlags(DiscardResource.Flags, EResourceTransitionFlags::Discard))
+		else
+#endif // #if PLATFORM_REQUIRES_TYPELESS_RESOURCE_DISCARD_WORKAROUND
 		{
-#if PLATFORM_REQUIRES_TYPELESS_RESOURCE_DISCARD_WORKAROUND
-			if (DiscardResource.Texture && DiscardResource.RTV && DiscardResource.Texture->GetRequiresTypelessResourceDiscardWorkaround())
+			if (DiscardResource.Subresource == D3D12_RESOURCE_BARRIER_ALL_SUBRESOURCES)
 			{
-				FLinearColor ClearColor = DiscardResource.Texture->GetClearColor();
-				GraphicsCommandList()->ClearRenderTargetView(DiscardResource.RTV->GetOfflineCpuHandle(), reinterpret_cast<float*>(&ClearColor), 0, nullptr);
-				UpdateResidency(DiscardResource.RTV->GetResource());
+				GraphicsCommandList()->DiscardResource(DiscardResource.Resource->GetResource(), nullptr);
 			}
 			else
-#endif // #if PLATFORM_REQUIRES_TYPELESS_RESOURCE_DISCARD_WORKAROUND
 			{
-				if (DiscardResource.Subresource == D3D12_RESOURCE_BARRIER_ALL_SUBRESOURCES)
-				{
-					GraphicsCommandList()->DiscardResource(DiscardResource.Resource->GetResource(), nullptr);
-				}
-				else
-				{
-					D3D12_DISCARD_REGION Region;
-					Region.NumRects = 0;
-					Region.pRects = nullptr;
-					Region.FirstSubresource = DiscardResource.Subresource;
-					Region.NumSubresources = 1;
+				D3D12_DISCARD_REGION Region;
+				Region.NumRects = 0;
+				Region.pRects = nullptr;
+				Region.FirstSubresource = DiscardResource.Subresource;
+				Region.NumSubresources = 1;
 
-					GraphicsCommandList()->DiscardResource(DiscardResource.Resource->GetResource(), &Region);
-				}
+				GraphicsCommandList()->DiscardResource(DiscardResource.Resource->GetResource(), &Region);
 			}
 		}
 	}
@@ -565,77 +551,84 @@ void FD3D12CommandContext::HandleTransientAliasing(const FD3D12TransitionData* T
 
 void FD3D12CommandContext::HandleResourceTransitions(const FD3D12TransitionData* TransitionData, bool& bUAVBarrier)
 {
-	const bool bDstAllPipelines = EnumHasAllFlags(TransitionData->DstPipelines, ERHIPipeline::All);
-	const bool bAsyncToAllPipelines = TransitionData->bAsyncToAllPipelines;
-
-	ERHIAccess SplitAccessBefore = ERHIAccess::Unknown;
-
 	for (const FRHITransitionInfo& Info : TransitionData->TransitionInfos)
 	{
+		if (!Info.Resource)
+		{
+			continue;
+		}
+
+		UE::RHICore::FResourceState ResourceState(*this, TransitionData->SrcPipelines, TransitionData->DstPipelines, Info);
+
+		// The transition SRV->UAV needs a UAV barrier then to work correctly otherwise there is no synchronization at all
 		bUAVBarrier |=
-			EnumHasAnyFlags(Info.AccessBefore, ERHIAccess::UAVMask)
-			&& EnumHasAnyFlags(Info.AccessAfter, ERHIAccess::UAVMask);
+			EnumHasAnyFlags(ResourceState.AccessBefore, ERHIAccess::UAVMask) &&
+			EnumHasAnyFlags(ResourceState.AccessAfter,  ERHIAccess::UAVMask);
+
+		// Skip duplicate transitions. This happens most frequently with implicit ones from NeedsExtraTransitions.
+		if (ResourceState.AccessBefore == ResourceState.AccessAfter)
+		{
+			continue;
+		}
 
 		// Very specific case that needs to be removed with EB
 		// a UAV -> SRVMAsk on the AsyncPipe get split in two: UAV->SRVCompute on Async and SRVCompute->SRVMask on Gfx
 		// On the Async pipe is going to be : UAV->SRVMask(that is automatically converted in UAV->SRVCompute)
 		// On the Direct(Gfx) pipe instead needs to be SRVCompute->SRVMask therefore the check there to change the Before state only on the Direct pipe.
-		SplitAccessBefore = ERHIAccess::Unknown;
-		if ((bAsyncToAllPipelines) && (Info.AccessAfter == ERHIAccess::SRVMask) && (QueueType == ED3D12QueueType::Direct))
+		if (TransitionData->bAsyncToAllPipelines && (ResourceState.AccessAfter == ERHIAccess::SRVMask) && (QueueType == ED3D12QueueType::Direct))
 		{
-			SplitAccessBefore = ERHIAccess::SRVCompute;
+			ResourceState.AccessBefore = ERHIAccess::SRVCompute;
 		}
 
-		const bool bImpossibleDiscardTransitionOnAsync = (Info.AccessAfter == ERHIAccess::Discard) && IsImpossibleToDiscardOnAsync(*this, Info);
-
 		// Process transitions which are forced during begin because those contain transition from Graphics to Compute and should
 		// help remove forced patch up command lists for async compute to run on the graphics queue
-		if (Info.Resource && (!bImpossibleDiscardTransitionOnAsync))
+		ProcessResource(*this, Info, [&](const FRHITransitionInfo& Info, FD3D12Resource* Resource, FD3D12Texture* Texture, FD3D12Texture* DiscardTexture)
 		{
-			ProcessResource(*this, Info, [&](const FRHITransitionInfo& Info, FD3D12Resource* Resource, FD3D12Texture* Texture, FD3D12Texture* DiscardTexture)
-				{
-					if (Resource == nullptr || !Resource->RequiresResourceStateTracking())
-					{
-						return;
-					}
+			if (!Resource->RequiresResourceStateTracking())
+			{
+				return;
+			}
 
-					check(Info.AccessBefore != ERHIAccess::Unknown);
-					check(Info.AccessAfter != ERHIAccess::Unknown);
+			if (ResourceState.AccessAfter == ERHIAccess::Discard && IsImpossibleAsyncDiscardTransition(ResourceState.DstPipelines, Texture))
+			{
+				return;
+			}
 
-					// special case that can happen during first frames or resize
-					if ((Info.AccessBefore == ERHIAccess::Present) && (Info.AccessAfter == ERHIAccess::Present))
-						return;
+			D3D12_RESOURCE_STATES StateBefore;
 
-					const ERHIAccess AccessBefore = SplitAccessBefore != ERHIAccess::Unknown ? SplitAccessBefore : Info.AccessBefore;
-						
-					D3D12_RESOURCE_STATES StateBefore = AccessBefore == ERHIAccess::Discard ? GetInitialResourceState(Resource->GetDesc(), QueueType) : GetD3D12ResourceState(Info.AccessBefore, QueueType, Texture);
-						
-					if (Info.AccessBefore != ERHIAccess::Present)
-					{
-						check(StateBefore != D3D12_RESOURCE_STATE_COMMON);
-					}
+			if (EnumHasAnyFlags(ResourceState.AccessBefore, ERHIAccess::Discard))
+			{
+				StateBefore = GetInitialResourceState(Resource->GetDesc(), QueueType);
+			}
+			else
+			{
+				StateBefore = GetD3D12ResourceState(ResourceState.AccessBefore, QueueType, Texture);
+			}
 
-					D3D12_RESOURCE_STATES StateAfter = Info.AccessAfter == ERHIAccess::Discard ? GetInitialResourceState(Resource->GetDesc(), QueueType) : GetD3D12ResourceState(Info.AccessAfter, QueueType, Texture);
+			if (ResourceState.AccessBefore != ERHIAccess::Present)
+			{
+				check(StateBefore != D3D12_RESOURCE_STATE_COMMON);
+			}
 
-					// enqueue the correct transitions
-					if (Info.IsWholeResource() || Resource->GetSubresourceCount() == 1)
-					{
-						TransitionResource(Resource, StateBefore, StateAfter, D3D12_RESOURCE_BARRIER_ALL_SUBRESOURCES);
-					}
-					else
-					{
-						// high level rendering is controlling transition ranges, at this level this is an index not a range
-						check(Info.MipIndex != FRHISubresourceRange::kAllSubresources);
-						check(Info.ArraySlice != FRHISubresourceRange::kAllSubresources);
-						check(Info.PlaneSlice != FRHISubresourceRange::kAllSubresources);
-						const uint32 Subresource = D3D12CalcSubresource(Info.MipIndex, Info.ArraySlice, Info.PlaneSlice, Resource->GetMipLevels(), Resource->GetArraySize());
-						check(Subresource < Resource->GetSubresourceCount());
-						TransitionResource(Resource, StateBefore, StateAfter, Subresource);
-					}
-				});
-		}
+			D3D12_RESOURCE_STATES StateAfter = ResourceState.AccessAfter == ERHIAccess::Discard ? GetInitialResourceState(Resource->GetDesc(), QueueType) : GetD3D12ResourceState(ResourceState.AccessAfter & ~ERHIAccess::Discard, QueueType, Texture);
+
+			// enqueue the correct transitions
+			if (Info.IsWholeResource() || Resource->GetSubresourceCount() == 1)
+			{
+				TransitionResource(Resource, StateBefore, StateAfter, D3D12_RESOURCE_BARRIER_ALL_SUBRESOURCES);
+			}
+			else
+			{
+				// high level rendering is controlling transition ranges, at this level this is an index not a range
+				check(Info.MipIndex != FRHISubresourceRange::kAllSubresources);
+				check(Info.ArraySlice != FRHISubresourceRange::kAllSubresources);
+				check(Info.PlaneSlice != FRHISubresourceRange::kAllSubresources);
+				const uint32 Subresource = D3D12CalcSubresource(Info.MipIndex, Info.ArraySlice, Info.PlaneSlice, Resource->GetMipLevels(), Resource->GetArraySize());
+				check(Subresource < Resource->GetSubresourceCount());
+				TransitionResource(Resource, StateBefore, StateAfter, Subresource);
+			}
+		});
 	}
-	
 }
 
 void FD3D12CommandContext::RHIBeginTransitions(TArrayView<const FRHITransition*> Transitions)
diff --git a/Engine/Source/Runtime/D3D12RHI/Private/D3D12RHI.cpp b/Engine/Source/Runtime/D3D12RHI/Private/D3D12RHI.cpp
--- a/Engine/Source/Runtime/D3D12RHI/Private/D3D12RHI.cpp
+++ b/Engine/Source/Runtime/D3D12RHI/Private/D3D12RHI.cpp
@@ -253,8 +253,8 @@ FD3D12DynamicRHI::FD3D12DynamicRHI(const TArray<TSharedPtr<FD3D12Adapter>>& Chos
 	// Enable run-time support if corresponding bit is set in DDSPI.
 	GRHIGlobals.SupportsVertexShaderUAVs = true;
 
-	GRHIGlobals.NeedExtraTransitions = true;
-	GRHIGlobals.EnableHighLevelTransitionPatching = true;
+	GRHIGlobals.NeedsExtraTransitions = true;
+	GRHIGlobals.NeedsTransientDiscardStateTracking = true;
 	GRHIGlobals.NeedsTransientDiscardOnGraphicsWorkaround = true;
 }
 
diff --git a/Engine/Source/Runtime/D3D12RHI/Private/D3D12Texture.cpp b/Engine/Source/Runtime/D3D12RHI/Private/D3D12Texture.cpp
--- a/Engine/Source/Runtime/D3D12RHI/Private/D3D12Texture.cpp
+++ b/Engine/Source/Runtime/D3D12RHI/Private/D3D12Texture.cpp
@@ -755,8 +755,6 @@ FD3D12DynamicRHI::FCreateTextureInternalResult FD3D12DynamicRHI::CreateTextureIn
 	Result.CreateState = CreateState;
 	Result.DesiredState = DesiredState;
 
-	check( (D3D12TextureOut->GetTrackedSubResourceCount()) == (D3D12TextureOut->ResourceLocation.GetResource()->GetSubresourceCount()) );
-
 	return Result;
 #else
 	checkf(false, TEXT("XBOX_CODE_MERGE : Removed. The Xbox platform version should be used."));
diff --git a/Engine/Source/Runtime/Engine/Private/GPUSkinCache.cpp b/Engine/Source/Runtime/Engine/Private/GPUSkinCache.cpp
--- a/Engine/Source/Runtime/Engine/Private/GPUSkinCache.cpp
+++ b/Engine/Source/Runtime/Engine/Private/GPUSkinCache.cpp
@@ -1398,7 +1398,7 @@ void FGPUSkinCache::DispatchUpdateSkinTangentsVertexPass(FRHICommandList& RHICmd
 		RHICmdList.Transition({
 			DispatchData.GetTangentRWBuffer()->UpdateAccessState(ERHIAccess::UAVCompute),
 			StagingBuffer->UpdateAccessState(ERHIAccess::UAVCompute)
-		}, ERHITransitionCreateFlags::NoFence);
+		}, ERHITransitionCreateFlags::AllowDecayPipelines);
 	}
 
 	SetComputePipelineState(RHICmdList, ComputeShader.GetComputeShader());
@@ -1500,7 +1500,7 @@ void FGPUSkinCache::DispatchUpdateSkinTangentsTrianglePass(FRHICommandList& RHIC
 		RHICmdList.Transition({
 			DispatchData.GetActiveTangentRWBuffer()->UpdateAccessState(ERHIAccess::SRVCompute),
 			StagingBuffer->UpdateAccessState(ERHIAccess::UAVCompute)
-		}, ERHITransitionCreateFlags::NoFence);
+		}, ERHITransitionCreateFlags::AllowDecayPipelines);
 	}
 
 	INC_DWORD_STAT_BY(STAT_GPUSkinCache_NumTrianglesForRecomputeTangents, NumTriangles);
@@ -1543,7 +1543,7 @@ void FGPUSkinCache::TransitionBuffers(FRHICommandList& RHICmdList, TArray<FSkinC
 		}
 
 		// The NoFence flag is necessary to silence the validator for transitioning from All pipes to Graphics.
-		RHICmdList.Transition(MakeArrayView(UAVs.GetData(), UAVs.Num()), ERHITransitionCreateFlags::NoFence);
+		RHICmdList.Transition(MakeArrayView(UAVs.GetData(), UAVs.Num()), ERHITransitionCreateFlags::AllowDecayPipelines);
 	}
 }
 
@@ -1570,7 +1570,7 @@ void FGPUSkinCache::TransitionBufferUAVs(FRHICommandList& RHICmdList, const TArr
 		}
 
 		// The NoFence flag is necessary to silence the validator for transitioning from All pipes to Graphics.
-		RHICmdList.Transition(MakeArrayView(UAVs.GetData(), UAVs.Num()), ERHITransitionCreateFlags::NoFence);
+		RHICmdList.Transition(MakeArrayView(UAVs.GetData(), UAVs.Num()), ERHITransitionCreateFlags::AllowDecayPipelines);
 	}
 }
 
@@ -2125,7 +2125,7 @@ void FGPUSkinCache::DispatchPassExecute(FTaskData& TaskData)
 		TaskData.AsyncComputeTransition = RHICreateTransition({ ERHIPipeline::AsyncCompute, ERHIPipeline::All, ERHITransitionCreateFlags::None, TransitionInfos });
 		RHICmdList.BeginTransition(TaskData.AsyncComputeTransition);
 		RHICmdList.EndTransition(TaskData.AsyncComputeTransition);
-		RHICmdList.SetTrackedAccess(TransitionInfos);
+		RHICmdList.SetTrackedAccess(TransitionInfos, ERHIPipeline::All);
 	}
 
 	BatchDispatches.Reset();
diff --git a/Engine/Source/Runtime/Engine/Private/VT/VirtualTextureUploadCache.cpp b/Engine/Source/Runtime/Engine/Private/VT/VirtualTextureUploadCache.cpp
--- a/Engine/Source/Runtime/Engine/Private/VT/VirtualTextureUploadCache.cpp
+++ b/Engine/Source/Runtime/Engine/Private/VT/VirtualTextureUploadCache.cpp
@@ -411,7 +411,7 @@ void FVirtualTextureUploadCache::FinalizeWithLegacyCopyTexture(FRDGBuilder& Grap
 			}
 		
 			RHICmdList.UnlockTexture(LockArgs);
-			RHICmdList.Transition(FRHITransitionInfo(StagingTexture.RHITexture, ERHIAccess::SRVMask, ERHIAccess::CopySrc));
+			RHICmdList.Transition(FRHITransitionInfo(StagingTexture.RHITexture, ERHIAccess::SRVMask, ERHIAccess::CopySrc), ERHITransitionCreateFlags::AllowDecayPipelines);
 		
 			// Upload each tile from staging texture to physical texture
 			for (uint32 Index = 0u; Index < BatchCount; ++Index)
@@ -489,10 +489,12 @@ void FVirtualTextureUploadCache::FinalizeWithUpdateTexture(FRDGBuilder& GraphBui
 			[PoolTileSize, BatchCount, TileEntries, UploadBuffers](FRDGAsyncTask,  FRHICommandList& RHICmdList)
 		{
 			SCOPED_GPU_MASK(RHICmdList, FRHIGPUMask::All());
-
 			SCOPED_DRAW_EVENT(RHICmdList, FVirtualTextureUploadCache_Finalize);
 			SCOPE_CYCLE_COUNTER(STAT_VTP_FlushUpload)
 
+			// These are already in the correct state from RDG, so disable the automatic RHI transitions.
+			FRHICommandListScopedAllowExtraTransitions ScopedExtraTransitions(RHICmdList, false);
+
 			for (uint32 Index = 0u; Index < BatchCount; ++Index)
 			{
 				const FTileEntry& Entry = TileEntries[Index];
diff --git a/Engine/Source/Runtime/RHI/Private/DynamicRHI.cpp b/Engine/Source/Runtime/RHI/Private/DynamicRHI.cpp
--- a/Engine/Source/Runtime/RHI/Private/DynamicRHI.cpp
+++ b/Engine/Source/Runtime/RHI/Private/DynamicRHI.cpp
@@ -283,8 +283,6 @@ void RHIInit(bool bHasEditorToken)
 		FGenericDataDrivenShaderPlatformInfo::Initialize();
 
 		GRHICommandList.LatchBypass(); // read commandline for bypass flag
-		// Platform RHIs can overidde this
-		GRHIGlobals.EnableHighLevelTransitionPatching = FParse::Param(FCommandLine::Get(), TEXT("EnableHighLevelTransitionsPatching"));
 
 		if (!FApp::CanEverRender())
 		{
@@ -813,20 +811,8 @@ FRenderQueryPoolRHIRef RHICreateRenderQueryPool(ERenderQueryType QueryType, uint
 FRHITransition* RHICreateTransition(const FRHITransitionCreateInfo& CreateInfo)
 {
 	// Placement create the transition on the heap.
-	FRHITransition* Transition = new (FRHICmdListBaseLinearAllocator::Malloc(FRHITransition::GetTotalAllocationSize(), (uint32)FRHITransition::GetAlignment())) FRHITransition(CreateInfo.SrcPipelines, CreateInfo.DstPipelines,
-		EnumHasAnyFlags(CreateInfo.Flags, ERHITransitionCreateFlags::AllowDuringRenderPass) ? ETransitionFlag::AllowDuringRenderPass : ETransitionFlag::None);
-
-	if (GRHIGlobals.EnableHighLevelTransitionPatching)
-	{
-		// Call to GDynamicRHI->RHICreateTransition will happen after the transitions have been tracked and patched, we are deep copying the CreateInfo here
-		Transition->CreationInfoData = new (FRHICmdListBaseLinearAllocator::Malloc(sizeof(FRHITransitionCreateInfoData), 16)) FRHITransitionCreateInfoData(CreateInfo);
-	}
-	else
-	{
-		EnumAddFlags(Transition->Flags, ETransitionFlag::PlatformLevelCreated);
-		GDynamicRHI->RHICreateTransition(Transition, CreateInfo);
-	}
-
+	FRHITransition* Transition = new (FRHICmdListBaseLinearAllocator::Malloc(FRHITransition::GetTotalAllocationSize(), (uint32)FRHITransition::GetAlignment())) FRHITransition(CreateInfo.SrcPipelines, CreateInfo.DstPipelines, CreateInfo.Flags);
+	GDynamicRHI->RHICreateTransition(Transition, CreateInfo);
 	return Transition;
 }
 
diff --git a/Engine/Source/Runtime/RHI/Private/RHI.cpp b/Engine/Source/Runtime/RHI/Private/RHI.cpp
--- a/Engine/Source/Runtime/RHI/Private/RHI.cpp
+++ b/Engine/Source/Runtime/RHI/Private/RHI.cpp
@@ -1245,6 +1245,8 @@ ERHIAccess GRHIMergeableAccessMask = ERHIAccess::ReadOnlyMask | ERHIAccess::UAVM
 // By default, only exclusively read only accesses are allowed.
 ERHIAccess GRHIMultiPipelineMergeableAccessMask = ERHIAccess::ReadOnlyExclusiveMask;
 
+ERHIAccess GRHIMultiSubresourceDiscardIntermediateAccess = ERHIAccess::SRVCompute;
+
 void FRHIDrawStats::Accumulate(FRHIDrawStats& Other)
 {
 	for (uint32 GPUIndex = 0; GPUIndex < GNumExplicitGPUsForRendering; ++GPUIndex)
diff --git a/Engine/Source/Runtime/RHI/Private/RHICommandList.cpp b/Engine/Source/Runtime/RHI/Private/RHICommandList.cpp
--- a/Engine/Source/Runtime/RHI/Private/RHICommandList.cpp
+++ b/Engine/Source/Runtime/RHI/Private/RHICommandList.cpp
@@ -1020,8 +1020,6 @@ void FRHICommandListExecutor::FSubmitState::Dispatch(FRHICommandListBase* CmdLis
 		//
 		Prereqs.Add(GRHICommandList.LastSubmit);
 
-		CmdList->UpdateTransitionsAndResourceState();
-		
 		CurrentTranslateJob->EnqueueTranslateTask(MoveTemp(Prereqs),
 			[TranslateJob = CurrentTranslateJob, CmdList]()
 			{
@@ -1937,7 +1935,7 @@ int32 FRHIComputeCommandList::GetGProfileGPUTransitions()
 void FRHIComputeCommandList::Transition(TArrayView<const FRHITransitionInfo> Infos, ERHITransitionCreateFlags CreateFlags)
 {
 	TransitionInternal(Infos, CreateFlags);
-	SetTrackedAccess(Infos);
+	SetTrackedAccess(Infos, GetPipelines());
 }
 
 void FRHIComputeCommandList::Transition(TArrayView<const FRHITransitionInfo> Infos, ERHIPipeline SrcPipelines, ERHIPipeline DstPipelines, ERHITransitionCreateFlags TransitionCreateFlags)
@@ -1972,7 +1970,7 @@ void FRHIComputeCommandList::Transition(TArrayView<const FRHITransitionInfo> Inf
 	{
 		// Set the tracked access on only one of the destination pipes.
 		FRHICommandListScopedPipeline Scope(*this, DstPipelines == ERHIPipeline::AsyncCompute ? ERHIPipeline::AsyncCompute : ERHIPipeline::Graphics);
-		SetTrackedAccess(Infos);
+		SetTrackedAccess(Infos, DstPipelines);
 	}
 }
 
@@ -2344,29 +2342,15 @@ void FRHICommandListBase::InvalidBufferFatalError(const FRHIBufferCreateDesc& Cr
 void FRHICommandListBase::TransitionInternal(TConstArrayView<FRHITransitionInfo> Infos, ERHITransitionCreateFlags CreateFlags)
 {
 	const ERHIPipeline Pipeline = GetPipeline();
-	
+	CreateFlags |= ERHITransitionCreateFlags::NoSplit;
+
 	if (Bypass())
 	{
-		// Without the FRHITransition we need to copy the transition info data because we could patch those data
-		TArray<FRHITransitionInfo, TSizedInlineAllocator<8, 32, FRHICmdListBaseArrayAllocator> > LocalTransitionInfo;
-		for (int32 i = 0; i < Infos.Num(); i++)
-		{
-			LocalTransitionInfo.Add(Infos[i]);
-
-			if (NeedTransitionsPatching())
-			{
-				// Because we know the Begin and End are executed consecutively, we don't have to worry about RenderPasses in between for NeedAllPlanes
-				const bool bForceTrackingUpdate = true;
-				PatchTransitionsInPlace((FRHITransitionInfo*)&LocalTransitionInfo[i], Pipeline, Pipeline, bForceTrackingUpdate);
-			}
-		}
-
 		// Stack allocate the transition
 		FMemStack& MemStack = FMemStack::Get();
 		FMemMark Mark(MemStack);
-		FRHITransition* Transition = new (MemStack.Alloc(FRHITransition::GetTotalAllocationSize(), FRHITransition::GetAlignment())) FRHITransition(Pipeline, Pipeline, 
-			EnumHasAnyFlags(CreateFlags, ERHITransitionCreateFlags::AllowDuringRenderPass) ? ETransitionFlag::AllowDuringRenderPass : ETransitionFlag::None);
-		GDynamicRHI->RHICreateTransition(Transition, FRHITransitionCreateInfo(Pipeline, Pipeline, CreateFlags | ERHITransitionCreateFlags::NoSplit, MakeArrayView(LocalTransitionInfo)));
+		FRHITransition* Transition = new (MemStack.Alloc(FRHITransition::GetTotalAllocationSize(), FRHITransition::GetAlignment())) FRHITransition(Pipeline, Pipeline, CreateFlags);
+		GDynamicRHI->RHICreateTransition(Transition, FRHITransitionCreateInfo(Pipeline, Pipeline, CreateFlags | ERHITransitionCreateFlags::NoSplit, Infos));
 
 		GetComputeContext().RHIBeginTransitions(MakeArrayView((const FRHITransition**)&Transition, 1));
 		GetComputeContext().RHIEndTransitions(MakeArrayView((const FRHITransition**)&Transition, 1));
@@ -2377,9 +2361,7 @@ void FRHICommandListBase::TransitionInternal(TConstArrayView<FRHITransitionInfo>
 	}
 	else
 	{
-		FRHITransition* Transition = RHICreateTransition(FRHITransitionCreateInfo(Pipeline, Pipeline, CreateFlags | ERHITransitionCreateFlags::NoSplit, Infos));
-		TrackResourceStateBegin(MakeArrayView((const FRHITransition**)&Transition, 1));
-		TrackResourceStateEnd(MakeArrayView((const FRHITransition**)&Transition, 1));
+		FRHITransition* Transition = RHICreateTransition(FRHITransitionCreateInfo(Pipeline, Pipeline, CreateFlags, Infos));
 		ALLOC_COMMAND(FRHICommandResourceTransition)(Transition);
 	}
 }
@@ -2481,414 +2463,3 @@ void FRHIComputeCommandList::RayTraceDispatchIndirect(
 
 	GlobalResourceBindings.Reset();
 }
-
-// NeedsAllPlanes currently only applies to texture with depth+stencil
-bool FRHICommandListBase::TrackTransitioningResource(const FRHITransitionCreateInfoData& TransitionCreateInfoData, const FRHITransitionInfo& TransitionInfo)
-{
-	if (GRHIGlobals.NeedsAllPlanes && GRHIGlobals.EnableHighLevelTransitionPatching)
-	{
-		// NoSplit means we don't have to worry about render passes between Begin and End
-		// Only applies to pipelines that remain on the graphics queue (for render passes)
-		// Only applies to depth textures with a stencil
-		// Only worth tracking if it's not an unknown/discard state
-		if (!EnumHasAnyFlags(TransitionCreateInfoData.m_Flags, ERHITransitionCreateFlags::NoSplit) &&
-			EnumHasAnyFlags(TransitionCreateInfoData.SrcPipelines, ERHIPipeline::Graphics) &&
-			EnumHasAnyFlags(TransitionCreateInfoData.DstPipelines, ERHIPipeline::Graphics) &&
-			(TransitionInfo.Type == FRHITransitionInfo::EType::Texture) &&
-			(TransitionInfo.AccessBefore != ERHIAccess::Unknown) && 
-			(TransitionInfo.AccessBefore != ERHIAccess::Discard))
-		{
-			FRHITexture* Texture = (FRHITexture*)GetViewableResource(TransitionInfo);
-			if (Texture && IsStencilFormat(Texture->GetDesc().Format))
-			{
-				return true;
-			}
-		}
-	}
-	return false;
-}
-
-// track all transitions:
-// - if a BeforeAccess is Unknown add the transition to a list for patching before dispatch
-// - save each resource state to use at dispatch time
-void FRHICommandListBase::TrackResourceStateBegin(TArrayView<const FRHITransition*> Transitions)
-{
-	// If cmdlist tracking is off or RHI can handle Unknown resource, skip the tracking
-	if (NeedTransitionsPatching() == false)
-	{
-		return;
-	}
-
-	for (const FRHITransition* Transition : Transitions)
-	{
-		FRHITransitionCreateInfoData* TransitionCreateInfoData = Transition->CreationInfoData;
-
-		if (TransitionCreateInfoData == nullptr)
-		{
-			checkSlow(Transition->IsPlatformLevelCreated());
-			continue;
-		}
-
-		for (int32 i = 0; i < TransitionCreateInfoData->TransitionInfos.Num(); i++)
-		{
-			FRHITransitionInfo* TransitionInfo = (FRHITransitionInfo*)&TransitionCreateInfoData->TransitionInfos[i];
-
-			if ((TransitionInfo->Type == FRHITransitionInfo::EType::Buffer) && (GRHIGlobals.SupportImplicitBufferStateTracking))
-			{ 
-				// This RHI supports Buffers state tracking at platform level, so we don't need to track them
-				continue;
-			}
-
-			FRHIViewableResource* TransitionViewableResource = GetViewableResource(*TransitionInfo);
-			if (TransitionViewableResource == nullptr)
-			{
-				continue;
-			}
-
-			TRHIResourceState<FRHICmdListBaseArrayAllocator>* TrackedResourceState = TrackedResources.Find(TransitionViewableResource);
-			const FRHISubresourceRange SubResourceIndex = *TransitionInfo;
-
-			// Both Unknown is not a supported case
-			check(!((TransitionInfo->AccessBefore == ERHIAccess::Unknown) && (TransitionInfo->AccessAfter == ERHIAccess::Unknown)));
-			// Discard to Unknown is not a supported case
-			check(!((TransitionInfo->AccessBefore == ERHIAccess::Discard) && (TransitionInfo->AccessAfter == ERHIAccess::Unknown)));
-
-			// Some RHI needs extra information on transition from Discard
-			const bool bBeforeIsDiscard = NeedsExtraTransitions() && (TransitionInfo->AccessBefore == ERHIAccess::Discard);
-
-			if ((TransitionInfo->AccessBefore == ERHIAccess::Unknown) || (TransitionInfo->AccessAfter == ERHIAccess::Unknown) || bBeforeIsDiscard)
-			{
-				ERHIAccess PreviousState = ERHIAccess::Unknown;
-				// Get the status from the previous state tracked
-				if (TrackedResourceState)
-				{
-					PreviousState = TrackedResourceState->GetSubresourceState(SubResourceIndex);
-				}
-
-				// Note: even if PreviousState is a valid state we cannot write into TransitionInfo since we are recording the cmdlist
-				TransitionInfoToPatch.Push(FTransitionPatchInfo(TransitionInfo, PreviousState));
-			}
-	
-			// Some RHI will not transition to initial state on After = Discard when running on the Async queue
-			bool bDiscardOnAsync = false;
-			if (NeedsExtraTransitions())
-			{
-				if ((TransitionInfo->AccessAfter == ERHIAccess::Discard)
-					&& ((TransitionCreateInfoData->SrcPipelines == ERHIPipeline::AsyncCompute) || (TransitionCreateInfoData->SrcPipelines == ERHIPipeline::All))
-					&& (TransitionCreateInfoData->DstPipelines == ERHIPipeline::AsyncCompute))
-				{
-					bDiscardOnAsync = true;
-				}
-			}
-			const bool bShouldSkipAfterState = ((EnumHasAnyFlags(TransitionInfo->Flags, EResourceTransitionFlags::IgnoreAfterState)) || (TransitionInfo->AccessAfter == ERHIAccess::Unknown));
-			const bool bNeedUpdateAfterState = (bShouldSkipAfterState == false) && (bDiscardOnAsync == false);
-
-			if (bNeedUpdateAfterState)
-			{
-				if (TrackTransitioningResource(*TransitionCreateInfoData, *TransitionInfo))
-				{
-					if (ActiveTransitionContainer* TransitionArray = ActiveTransitions.Find(TransitionViewableResource))
-					{
-						TransitionArray->Add(TransitionInfo);
-					}
-					else
-					{
-						ActiveTransitionContainer NewArray;
-						NewArray.Add(TransitionInfo);
-						ActiveTransitions.Add(TransitionViewableResource, MoveTemp(NewArray));
-					}
-				}
-
-				// Update the TrackedResource State
-				if (TrackedResourceState)
-				{
-					TrackedResourceState->SetSubresourceState(SubResourceIndex, TransitionInfo->AccessAfter);
-				}
-				else
-				{
-					// Add a new element
-					TrackedResources.Add(TransitionViewableResource, TRHIResourceState<FRHICmdListBaseArrayAllocator>(SubResourceIndex, TransitionInfo->AccessAfter, TransitionViewableResource->GetTrackedSubResourceCount(), true));
-				}
-			}
-
-			// Some RHI need all planes of a DepthStencil texture to transition
-			if (NeedsAllPlanes() &&
-				(TransitionInfo->Type == FRHITransitionInfo::EType::Texture) &&
-				!TransitionInfo->IsAllPlaneSlices())
-			{
-				// Only applies to texture formats with a stencil
-				FRHITexture* Texture = (FRHITexture*)TransitionViewableResource;
-				if (IsStencilFormat(TransitionInfo->Texture->GetDesc().Format))
-				{
-					// Set the OtherPlaneState for all single plane transitions, less work than tracking the set planes per subresource
-					const FRHISubresourceRange OtherResourceIndex(
-						SubResourceIndex.MipIndex,
-						SubResourceIndex.ArraySlice,
-						(SubResourceIndex.PlaneSlice == FRHISubresourceRange::kDepthPlaneSlice) ? FRHISubresourceRange::kStencilPlaneSlice : FRHISubresourceRange::kDepthPlaneSlice
-					);
-					const ERHIAccess OtherPlaneState = TrackedResourceState ? TrackedResourceState->GetSubresourceState(OtherResourceIndex) : ERHIAccess::Unknown;
-					PlaneTransitionInfoToPatch.Push(FTransitionPatchInfo(TransitionInfo, OtherPlaneState));
-				}
-			}
-		}
-	}
-}
-
-void FRHICommandListBase::TrackResourceStateEnd(TArrayView<const FRHITransition*> Transitions)
-{
-	if (!NeedTransitionsPatching() || !NeedsAllPlanes())
-	{
-		return;
-	}
-
-	for (const FRHITransition* Transition : Transitions)
-	{
-		const FRHITransitionCreateInfoData* TransitionCreateInfoData = Transition->CreationInfoData;
-
-		if (!TransitionCreateInfoData)
-		{
-			continue;
-		}
-
-		for (int32 Index = 0; Index < TransitionCreateInfoData->TransitionInfos.Num(); Index++)
-		{
-			const FRHITransitionInfo& TransitionInfo = TransitionCreateInfoData->TransitionInfos[Index];
-
-			if (!TrackTransitioningResource(*TransitionCreateInfoData, TransitionInfo))
-			{
-				continue;
-			}
-
-			FRHIViewableResource* TransitionViewableResource = GetViewableResource(TransitionInfo);
-			if (TransitionViewableResource == nullptr)
-			{
-				continue;
-			}
-
-			bool bWasFound = false;
-			if (ActiveTransitionContainer* TransitionArray = ActiveTransitions.Find(TransitionViewableResource))
-			{
-				const int32 NumRemoved = TransitionArray->RemoveSwap(&TransitionInfo);
-				bWasFound = (NumRemoved > 0);
-			}
-
-
-			if (!bWasFound)
-			{
-				if (ActiveTransitionContainer* TransitionArray = ClosedTransitions.Find(TransitionViewableResource))
-				{
-					TransitionArray->Add(&TransitionInfo);
-				}
-				else
-				{
-					ActiveTransitionContainer NewArray;
-					NewArray.Add(&TransitionInfo);
-					ClosedTransitions.Add(TransitionViewableResource, MoveTemp(NewArray));
-				}
-			}
-		}
-	}
-}
-
-TMap<FRHIViewableResource*, TArray<FRHITransitionInfo>> FRHICommandListBase::GlobalActiveTransitions;
-
-void FRHICommandListBase::UpdateTransitionsAndResourceState()
-{
-	TRACE_CPUPROFILER_EVENT_SCOPE(UpdateTransitionsAndResourceState);
-
-	//If cmdlist tracking is off or RHI can handle Unknown resource, skip the tracking
-	if (NeedTransitionsPatching())
-	{
-		// Iterate over all TransitionInfo to patch for Unknown
-		for (int32 i = 0; i < TransitionInfoToPatch.Num(); i++)
-		{
-			FTransitionPatchInfo& PatchInfo = TransitionInfoToPatch[i];
-			FRHITransitionInfo* TransitionInfo = PatchInfo.TransitionInfo;
-			FRHISubresourceRange SubResourceIndex = *TransitionInfo;
-
-			const bool bBeforeIsDiscard = NeedsExtraTransitions() && (TransitionInfo->AccessBefore == ERHIAccess::Discard);
-
-			if (bBeforeIsDiscard)
-			{ 
-				if (PatchInfo.State != ERHIAccess::Unknown)
-				{
-					TransitionInfo->DiscardAccessBefore = PatchInfo.State;
-				}
-				else
-				{
-					ERHIAccess ResourceState = ERHIAccess::Unknown;
-
-					FRHIViewableResource* TransitionViewableResource = GetViewableResource(*TransitionInfo);
-
-					if (TransitionViewableResource)
-					{
-						ResourceState = TransitionViewableResource->TrackedAccess.GetSubresourceState(SubResourceIndex);
-					}
-
-					TransitionInfo->DiscardAccessBefore = ResourceState;
-				}
-			}
-			else if (TransitionInfo->AccessBefore == ERHIAccess::Unknown)
-			{
-				if (PatchInfo.State != ERHIAccess::Unknown)
-				{
-					TransitionInfo->AccessBefore = PatchInfo.State;
-				}
-				else
-				{
-					ERHIAccess ResourceState = ERHIAccess::Unknown;
-
-					FRHIViewableResource* TransitionViewableResource = GetViewableResource(*TransitionInfo);
-
-					if (TransitionViewableResource)
-					{
-						ResourceState = TransitionViewableResource->TrackedAccess.GetSubresourceState(SubResourceIndex);
-					}
-
-					TransitionInfo->AccessBefore = ResourceState;
-				}
-			}
-			else if (TransitionInfo->AccessAfter == ERHIAccess::Unknown)
-			{
-				if (PatchInfo.State != ERHIAccess::Unknown)
-				{
-					TransitionInfo->AccessAfter = PatchInfo.State;
-				}
-				else
-				{
-					ERHIAccess ResourceState = ERHIAccess::Unknown;
-
-					FRHIViewableResource* TransitionViewableResource = GetViewableResource(*TransitionInfo);
-
-					if (TransitionViewableResource)
-					{
-						ResourceState = TransitionViewableResource->TrackedAccess.GetSubresourceState(SubResourceIndex);
-					}
-
-					TransitionInfo->AccessAfter = ResourceState;
-				}
-			}
-
-			check(TransitionInfo->AccessBefore != ERHIAccess::Unknown);
-			check(TransitionInfo->AccessAfter != ERHIAccess::Unknown);
-		}
-
-		if (NeedsAllPlanes())
-		{
-			// Iterate over all PlaneTransitionInfo to patch any missing plane information
-			for (FTransitionPatchInfo& PatchInfo : PlaneTransitionInfoToPatch)
-			{
-				PatchInfo.TransitionInfo->Flags |= EResourceTransitionFlags::FillOtherPlaneState;
-				if (PatchInfo.State != ERHIAccess::Unknown)
-				{
-					PatchInfo.TransitionInfo->OtherPlaneAccess = PatchInfo.State;
-				}
-				else
-				{
-					FRHIViewableResource* TransitionViewableResource = GetViewableResource(*PatchInfo.TransitionInfo);
-					if (TransitionViewableResource)
-					{
-						const FRHISubresourceRange OtherResourceIndex(
-							PatchInfo.TransitionInfo->MipIndex,
-							PatchInfo.TransitionInfo->ArraySlice,
-							(PatchInfo.TransitionInfo->PlaneSlice == FRHISubresourceRange::kDepthPlaneSlice) ? FRHISubresourceRange::kStencilPlaneSlice : FRHISubresourceRange::kDepthPlaneSlice
-						);
-						PatchInfo.TransitionInfo->OtherPlaneAccess = TransitionViewableResource->TrackedAccess.GetSubresourceState(OtherResourceIndex);
-					}
-				}
-			}
-			PlaneTransitionInfoToPatch.Empty();
-
-			// Iterate over all RenderPassInfo to patch any missing plane states
-			for (FRenderPassPatchInfo& PatchInfo : RenderPassInfoToPatch)
-			{
-				PatchInfo.RenderPassInfo->NopAccess = PatchInfo.State;
-
-				if (PatchInfo.RenderPassInfo->NopAccess == ERHIAccess::Unknown)
-				{
-					FRHITexture* DepthStencilTexture = PatchInfo.RenderPassInfo->DepthStencilRenderTarget.DepthStencilTarget;				
-					if (TArray<FRHITransitionInfo>* TransitionArray = GlobalActiveTransitions.Find(DepthStencilTexture))
-					{
-						for (const FRHITransitionInfo& TransitionInfo : *TransitionArray)
-						{
-							if (!PatchInfo.RenderPassInfo->DepthStencilRenderTarget.ExclusiveDepthStencil.IsUsingDepth())
-							{
-								if (TransitionInfo.IsAllPlaneSlices() || TransitionInfo.PlaneSlice == FRHISubresourceRange::kDepthPlaneSlice)
-								{
-									PatchInfo.RenderPassInfo->NopAccess = TransitionInfo.AccessBefore;
-									break;
-								}
-							}
-							else if (!PatchInfo.RenderPassInfo->DepthStencilRenderTarget.ExclusiveDepthStencil.IsUsingStencil())
-							{
-								if (TransitionInfo.IsAllPlaneSlices() || TransitionInfo.PlaneSlice == FRHISubresourceRange::kStencilPlaneSlice)
-								{
-									PatchInfo.RenderPassInfo->NopAccess = TransitionInfo.AccessBefore;
-									break;
-								}
-							}
-						}
-					}
-				}
-
-				if (PatchInfo.RenderPassInfo->NopAccess == ERHIAccess::Unknown)
-				{
-					ResolveNopAccessInPlace(*PatchInfo.RenderPassInfo);
-				}
-			}
-			RenderPassInfoToPatch.Empty();
-
-			for (const auto& Pair : ClosedTransitions)
-			{
-				if (TArray<FRHITransitionInfo>* GlobalTransitionArray = GlobalActiveTransitions.Find(Pair.Key))
-				{
-					for (const FRHITransitionInfo* ClosedTransition : Pair.Value)
-					{
-						const int32 RemoveCount = GlobalTransitionArray->RemoveSwap(*ClosedTransition);
-						if (RemoveCount == 0)
-						{
-							UE_LOG(LogRHI, Warning, TEXT("Could not close transition that didn't exist in the tracked GlobalActiveTransitions."));
-						}
-
-						if (GlobalTransitionArray->Num() == 0)
-						{
-							GlobalActiveTransitions.Remove(Pair.Key);
-						}
-					}
-				}
-			}
-			ClosedTransitions.Empty();
-
-			for (const auto& Pair : ActiveTransitions)
-			{
-				if (Pair.Value.Num() == 0)
-				{
-					continue;
-				}
-
-				TArray<FRHITransitionInfo>* GlobalTransitionArray = GlobalActiveTransitions.Find(Pair.Key);
-				if (!GlobalTransitionArray)
-				{
-					GlobalTransitionArray = &GlobalActiveTransitions.Add(Pair.Key);
-				}
-
-				for (const FRHITransitionInfo* ActiveTransition : Pair.Value)
-				{
-					GlobalTransitionArray->Add(*ActiveTransition);
-				}
-			}
-			ActiveTransitions.Empty();
-		}
-
-		// Iterate over all tracked Resources to update their TrackedState with the latest state
-		for (TPair<FRHIViewableResource*, TRHIResourceState<FRHICmdListBaseArrayAllocator>>& Pair : TrackedResources)
-		{
-			FRHIViewableResource* TrackedResource = Pair.Key;
-			TRHIResourceState<FRHICmdListBaseArrayAllocator>* TrackedResourceState = &Pair.Value;
-			TrackedResource->TrackedAccess.Merge(*TrackedResourceState);
-		}
-
-		TransitionInfoToPatch.Empty();
-		TrackedResources.Empty();
-	}
-}
diff --git a/Engine/Source/Runtime/RHI/Private/RHIResources.cpp b/Engine/Source/Runtime/RHI/Private/RHIResources.cpp
--- a/Engine/Source/Runtime/RHI/Private/RHIResources.cpp
+++ b/Engine/Source/Runtime/RHI/Private/RHIResources.cpp
@@ -123,7 +123,7 @@ bool FRHIResource::Bypass()
 }
 
 FRHITexture::FRHITexture(const FRHITextureCreateDesc& InDesc)
-	: FRHIViewableResource(RRT_Texture, InDesc.InitialState, InDesc.GetSubresourceCount())
+	: FRHIViewableResource(RRT_Texture, InDesc.InitialState)
 #if ENABLE_RHI_VALIDATION
 	, RHIValidation::FTextureResource(InDesc)
 #endif
diff --git a/Engine/Source/Runtime/RHI/Private/RHIStrings.cpp b/Engine/Source/Runtime/RHI/Private/RHIStrings.cpp
--- a/Engine/Source/Runtime/RHI/Private/RHIStrings.cpp
+++ b/Engine/Source/Runtime/RHI/Private/RHIStrings.cpp
@@ -503,10 +503,7 @@ FString GetResourceTransitionFlagsName(EResourceTransitionFlags Flags)
 		default: checkNoEntry(); // fall through
 		case EResourceTransitionFlags::None:                return TEXT("None");
 		case EResourceTransitionFlags::MaintainCompression: return TEXT("MaintainCompression");
-		case EResourceTransitionFlags::Discard:				return TEXT("Discard");
-		case EResourceTransitionFlags::Clear:				return TEXT("Clear");
 		case EResourceTransitionFlags::IgnoreAfterState:	return TEXT("IgnoreAfterState");
-		case EResourceTransitionFlags::FillOtherPlaneState:	return TEXT("FillOtherPlaneState");
 		}
 	});
 }
diff --git a/Engine/Source/Runtime/RHI/Private/RHITextureReference.cpp b/Engine/Source/Runtime/RHI/Private/RHITextureReference.cpp
--- a/Engine/Source/Runtime/RHI/Private/RHITextureReference.cpp
+++ b/Engine/Source/Runtime/RHI/Private/RHITextureReference.cpp
@@ -83,8 +83,8 @@ const FRHITextureDesc& FRHITextureReference::GetDesc() const
 	return ReferencedTexture->GetDesc();
 }
 
-void FRHITextureReference::SetTrackedAccessFromContext(ERHIAccess Access)
+void FRHITextureReference::SetTrackedAccessFromContext(FRHITrackedAccess InTrackedAccess)
 {
 	check(ReferencedTexture);
-	ReferencedTexture->SetTrackedAccessFromContext(Access);
-}
+	ReferencedTexture->SetTrackedAccessFromContext(InTrackedAccess);
+}
\ No newline at end of file
diff --git a/Engine/Source/Runtime/RHI/Private/RHITransition.cpp b/Engine/Source/Runtime/RHI/Private/RHITransition.cpp
--- a/Engine/Source/Runtime/RHI/Private/RHITransition.cpp
+++ b/Engine/Source/Runtime/RHI/Private/RHITransition.cpp
@@ -40,43 +40,10 @@ FRHIViewableResource* GetViewableResource(const FRHITransitionInfo& Info)
 	return nullptr;
 }
 
-void FRHITransition::CleanupInfoData() const
-{
-	checkSlow(CreationInfoData != nullptr);
-	CreationInfoData->~FRHITransitionCreateInfoData();
-	FRHICmdListBaseLinearAllocator::Free(CreationInfoData);
-	CreationInfoData = nullptr;
-}
-
-void FRHITransition::CreatePlatformTransition(bool bCleanupInfoData) const
-{
-	checkSlow(GRHIGlobals.EnableHighLevelTransitionPatching);
-	if (IsPlatformLevelCreated() == false)
-	{
-		GDynamicRHI->RHICreateTransition(const_cast<FRHITransition*>(this), CreationInfoData->GetCreationInfo());
-		EnumAddFlags(Flags, ETransitionFlag::PlatformLevelCreated);
-
-		if (bCleanupInfoData)
-		{
-			CleanupInfoData();
-		}
-	}
-}
-
-void FRHITransition::CreatePlatformTransitionThreadSafe(bool bCleanupInfoData) const
-{
-	checkSlow(GRHIGlobals.EnableHighLevelTransitionPatching);
-	
-	UE::TUniqueLock Lock(TransitionCreationMutex);
-
-	CreatePlatformTransition(bCleanupInfoData);
-}
-
 void FRHITransition::Cleanup() const
 {
 	FRHITransition* Transition = const_cast<FRHITransition*>(this);
 	RHIReleaseTransition(Transition);
-	checkSlow((CreationInfoData == nullptr) || GRHIGlobals.NeedsAllPlanes);
 
 	// Explicit destruction of the transition.
 	Transition->~FRHITransition();
diff --git a/Engine/Source/Runtime/RHI/Private/RHIValidation.cpp b/Engine/Source/Runtime/RHI/Private/RHIValidation.cpp
--- a/Engine/Source/Runtime/RHI/Private/RHIValidation.cpp
+++ b/Engine/Source/Runtime/RHI/Private/RHIValidation.cpp
@@ -924,7 +924,8 @@ void FValidationRHI::RHICreateTransition(FRHITransition* Transition, const FRHIT
 		if (!Info.Resource)
 			continue;
 
-		checkf(Info.AccessAfter != ERHIAccess::Unknown, TEXT("FRHITransitionInfo::AccessAfter cannot be Unknown when creating a resource transition."));
+		RHI_VALIDATION_CHECK(Info.AccessAfter != ERHIAccess::Unknown || SrcPipelines == DstPipelines && DstPipelines != ERHIPipeline::All, TEXT("Cannot use Unknown after state when transitioning between pipelines."));
+
 		checkf(Info.Type != FRHITransitionInfo::EType::Unknown, TEXT("FRHITransitionInfo::Type cannot be Unknown when creating a resource transition."));
 
 		if (const FRHICommitResourceInfo* CommitInfo = Info.CommitInfo.GetPtrOrNull())
@@ -1598,7 +1599,7 @@ namespace RHIValidation
 			*GetRHIPipelineName(RequiredState.Pipelines));
 	}
 
-	static inline FString GetReasonString_IncorrectTrackedAccess(
+	static inline FString GetReasonString_IncorrectSetTrackedAccess(
 		FResource* Resource, FSubresourceIndex const& SubresourceIndex,
 		const FState& CurrentState,
 		const FState& TrackedState)
@@ -1608,12 +1609,38 @@ namespace RHIValidation
 			BARRIER_TRACKER_LOG_PREFIX_RESNAME
 			TEXT("Attempted to assign resource %s a tracked access that does not match its validation tracked access.\n\n")
 			TEXT("    --- Actual access states:                    %s\n")
+			TEXT("    --- Actual pipelines:                        %s\n")
 			TEXT("    --- Assigned access states:                  %s\n")
+			TEXT("    --- Assigned pipelines:                      %s\n")
+			BARRIER_TRACKER_LOG_SUFFIX,
+			*DebugName,
+			*DebugName,
+			*GetRHIAccessName(CurrentState.Access),
+			*GetRHIPipelineName(CurrentState.Pipelines),
+			*GetRHIAccessName(TrackedState.Access),
+			*GetRHIPipelineName(TrackedState.Pipelines));
+	}
+	
+	static inline FString GetReasonString_IncorrectGetTrackedAccess(
+		FResource* Resource, FSubresourceIndex const& SubresourceIndex,
+		const FState& CurrentState,
+		const FState& TrackedState)
+	{
+		FString DebugName = GetResourceDebugName(Resource, SubresourceIndex);
+		return FString::Printf(
+			BARRIER_TRACKER_LOG_PREFIX_RESNAME
+			TEXT("Attempted to resolve ERHIAccess::Unknown for resource %s but its tracked access that does match its validation tracked access.\n\n")
+			TEXT("    --- Validation actual access states:        %s\n")
+			TEXT("    --- Validation actual pipelines:            %s\n")
+			TEXT("    --- Tracked access states:                  %s\n")
+			TEXT("    --- Tracked pipelines:                      %s\n")
 			BARRIER_TRACKER_LOG_SUFFIX,
 			*DebugName,
 			*DebugName,
 			*GetRHIAccessName(CurrentState.Access),
-			*GetRHIAccessName(TrackedState.Access));
+			*GetRHIPipelineName(CurrentState.Pipelines),
+			*GetRHIAccessName(TrackedState.Access),
+			*GetRHIPipelineName(TrackedState.Pipelines));
 	}
 
 	static FString ResolveAndFormatCallstack (uint64* Trace)
@@ -1961,7 +1988,7 @@ namespace RHIValidation
 			TEXT("    --- ERHIAccess::Unknown. RHICmdList.SetTrackedAccess applies to whole resources.\n")
 			BARRIER_TRACKER_LOG_SUFFIX,
 			*DebugName,
-			*GetRHIAccessName(Resource->GetTrackedAccess()),
+			*GetRHIAccessName(Resource->GetTrackedState().Access),
 			*GetRHIAccessName(CurrentState.Access),
 			*DebugName,
 			*GetRHIPipelineName(CurrentState.Pipelines),
@@ -2059,7 +2086,7 @@ namespace RHIValidation
 		FString DebugName = GetResourceDebugName(Resource, SubresourceIndex);
 		return FString::Printf(
 			BARRIER_TRACKER_LOG_PREFIX_RESNAME
-			TEXT("Attempted to begin a resource transition for resource %s on All pipes. Transitionwith  EResourceTransitionFlags::IgnoreAfterState on All pipes are not supported.\n\n")
+			TEXT("Attempted to begin a resource transition for resource %s on All pipes. Transition with EResourceTransitionFlags::IgnoreAfterState on All pipes are not supported.\n\n")
 			TEXT("    --- Pending access states for this resource are:              %s\n")
 			TEXT("    --- Attempted access states for the current transition are: %s\n")
 			TEXT("    --- Pending pipelines for this resource are:                  %s\n")
@@ -2238,10 +2265,60 @@ namespace RHIValidation
 		}
 	}
 
-	void FSubresourceState::BeginTransition(FResource* Resource, FSubresourceIndex const& SubresourceIndex, const FState& CurrentStateFromRHI, const FState& TargetState, EResourceTransitionFlags NewFlags, ERHITransitionCreateFlags CreateFlags, ERHIPipeline ExecutingPipeline, const TRHIPipelineArray<uint64>& PipelineMaxAwaitedFenceValues, void* CreateTrace)
+	void FResource::InitTransient(const TCHAR* InDebugName)
+	{
+		check(TransientState.bTransient && TransientState.Status != FTransientState::EStatus::Acquired);
+		TransientState.Status = FTransientState::EStatus::None;
+		DebugName = InDebugName;
+
+		for (ERHIPipeline Pipeline : MakeFlagsRange(ERHIPipeline::All))
+		{
+			auto& State = WholeResourceState.States[Pipeline];
+
+			State.Current.Access = ERHIAccess::Discard;
+			State.Current.Pipelines = Pipeline;
+			State.Previous = State.Current;
+		}
+		SubresourceStates.Reset();
+	}
+
+	void FResource::InitBarrierTracking(int32 InNumMips, int32 InNumArraySlices, int32 InNumPlanes, ERHIAccess InResourceState, const TCHAR* InDebugName)
+	{
+		checkSlow(InNumMips > 0 && InNumArraySlices > 0 && InNumPlanes > 0);
+		check(InResourceState != ERHIAccess::Unknown);
+
+		NumMips = InNumMips;
+		NumArraySlices = InNumArraySlices;
+		NumPlanes = InNumPlanes;
+		TransientState = FTransientState(InResourceState);
+		TrackedState = FState(InResourceState, ERHIPipeline::None);
+
+		for (ERHIPipeline Pipeline : MakeFlagsRange(ERHIPipeline::All))
+		{
+			auto& State = WholeResourceState.States[Pipeline];
+
+			State.Current.Access = InResourceState;
+			State.Current.Pipelines = Pipeline;
+			State.Previous = State.Current;
+		}
+
+		if (InDebugName != nullptr)
+		{
+			SetDebugName(InDebugName);
+		}
+	}
+
+	void FSubresourceState::BeginTransition(FResource* Resource, FSubresourceIndex const& SubresourceIndex, const FState& CurrentStateFromRHI, const FState& InTargetState, EResourceTransitionFlags NewFlags, ERHITransitionCreateFlags CreateFlags, ERHIPipeline ExecutingPipeline, const TRHIPipelineArray<uint64>& PipelineMaxAwaitedFenceValues, void* CreateTrace)
 	{
 		FPipelineState& State = States[ExecutingPipeline];
 
+		FState TargetState = InTargetState;
+
+		if (TargetState.Access == ERHIAccess::Unknown)
+		{
+			TargetState.Access = Resource->GetTrackedState().Access;
+		}
+
 		void* BeginTrace = nullptr;
 		if (Resource->LoggingMode != ELoggingMode::None 
 #if LOG_UNNAMED_RESOURCES
@@ -2249,20 +2326,33 @@ namespace RHIValidation
 #endif
 			)
 		{
-			BeginTrace = Log(Resource, SubresourceIndex, CreateTrace, TEXT("Begin"), TEXT("BeginTransition"), *FString::Printf(TEXT("Current: (%s) -> Before(%s) New: (%s), Flags: %s, Executing Pipeline: %s"),
+			const TCHAR* PulledFromTracked = TEXT("");
+
+			if (InTargetState.Access == ERHIAccess::Unknown)
+			{
+				PulledFromTracked = TEXT(" (Pulled From SetTrackedAccess)");
+			}
+
+			BeginTrace = Log(Resource, SubresourceIndex, CreateTrace, TEXT("Begin"), TEXT("BeginTransition"), *FString::Printf(TEXT("Current: (%s) -> Before(%s) New: (%s)%s, Flags: %s, Executing Pipeline: %s"),
 				*State.Current.ToString(),
 				*CurrentStateFromRHI.ToString(),
 				*TargetState.ToString(),
+				PulledFromTracked,
 				*GetResourceTransitionFlagsName(NewFlags),
 				*GetRHIPipelineName(ExecutingPipeline)
 			));
 		}
 
+		if (CurrentStateFromRHI.Access == ERHIAccess::Unknown)
+		{
+			RHI_VALIDATION_CHECK(Resource->GetTrackedState().Access == State.Previous.Access, *GetReasonString_IncorrectGetTrackedAccess(Resource, SubresourceIndex, State.Previous, Resource->GetTrackedState()));
+		}
+
 		if (Resource->TransientState.bTransient)
 		{
 			RHI_VALIDATION_CHECK(Resource->TransientState.IsAcquired(), *GetReasonString_TransitionWithoutAcquire(Resource));
 
-			if (TargetState.Access == ERHIAccess::Discard)
+			if (EnumHasAnyFlags(TargetState.Access, ERHIAccess::Discard))
 			{
 				Resource->TransientState.Discard(Resource, CreateTrace, CurrentStateFromRHI.Pipelines, ExecutingPipeline);
 			}
@@ -2286,14 +2376,20 @@ namespace RHIValidation
 			// Check for the correct pipeline
 			RHI_VALIDATION_CHECK(EnumHasAllFlags(CurrentStateFromRHI.Pipelines, ExecutingPipeline), *GetReasonString_WrongPipeline(Resource, SubresourceIndex, State.Current, TargetState));
 
-			const auto HasMatchingPipelines = [CreateFlags] (ERHIPipeline Previous, ERHIPipeline Next)
+			const auto HasMatchingPipelines = [CreateFlags, Resource] (ERHIPipeline PreviousFromState, ERHIPipeline PreviousFromRHI)
 			{
-				// If no fence is being issued we only need to validate that the transition is happening from one of the previous pipes.
-				if (EnumHasAnyFlags(CreateFlags, ERHITransitionCreateFlags::NoFence))
+				if (PreviousFromState == PreviousFromRHI)
 				{
-					return EnumHasAllFlags(Previous, Next);
+					return true;
 				}
-				return Previous == Next;
+
+				// We allow collapsing pipes from All -> Single only if the flag is explicitly provided.
+				if (EnumHasAnyFlags(CreateFlags, ERHITransitionCreateFlags::AllowDecayPipelines))
+				{
+					return EnumHasAnyFlags(PreviousFromState, PreviousFromRHI) && PreviousFromState == Resource->GetTrackedState().Pipelines;
+				}
+
+				return false;
 			};
 
 
@@ -2301,13 +2397,13 @@ namespace RHIValidation
 
 			// We do not check pipelines for IgnoreAfterState since we do not replicate those transition on all pipelines.
 			if (EnumHasAnyFlags(NewFlags, EResourceTransitionFlags::IgnoreAfterState) == false)
-			{ 
+			{
 				bHasMatchingPipelines = HasMatchingPipelines(State.Previous.Pipelines, CurrentStateFromRHI.Pipelines);
 			}
 
 			if (CurrentStateFromRHI.Access == ERHIAccess::Unknown)
 			{
-				RHI_VALIDATION_CHECK(Resource->TrackedAccess == State.Previous.Access && bHasMatchingPipelines,
+				RHI_VALIDATION_CHECK(Resource->TrackedState.Access == State.Previous.Access && bHasMatchingPipelines,
 					*GetReasonString_IncorrectPreviousTrackedState(Resource, SubresourceIndex, State.Previous, CurrentStateFromRHI.Pipelines));
 			}
 			else
@@ -2327,11 +2423,6 @@ namespace RHIValidation
 		RHI_VALIDATION_CHECK(bRegularTransitionWhileIgnoring == false,
 				*GetReasonString_MismatchedIgnoreAfterState(Resource, SubresourceIndex, State.Current, TargetState));
 
-		// Check for unnecessary transitions
-		// @todo: this check is not particularly useful at the moment, as there are many unnecessary resource transitions.
-		//RHI_VALIDATION_CHECK(CurrentState != TargetState, *GetReasonString_UnnecessaryTransition(Resource, SubresourceIndex, CurrentState));
-
-
 		State.bTransitioning = true;
 		State.Flags = NewFlags;
 		State.BeginTransitionBacktrace = BeginTrace;
@@ -2351,28 +2442,42 @@ namespace RHIValidation
 				States[OtherPipeline] = State;
 			}
 		}
-		
 	}
 
-	void FSubresourceState::EndTransition(FResource* Resource, FSubresourceIndex const& SubresourceIndex, const FState& CurrentStateFromRHI, const FState& TargetState, EResourceTransitionFlags NewFlags, ERHIPipeline ExecutingPipeline, uint64 ExecutingPipelineFenceValue, void* CreateTrace)
+	void FSubresourceState::EndTransition(FResource* Resource, FSubresourceIndex const& SubresourceIndex, const FState& CurrentStateFromRHI, const FState& InTargetState, EResourceTransitionFlags NewFlags, ERHIPipeline ExecutingPipeline, uint64 ExecutingPipelineFenceValue, void* CreateTrace)
 	{
+		FPipelineState& State = States[ExecutingPipeline];
+
+		FState TargetState = InTargetState;
+
+		if (TargetState.Access == ERHIAccess::Unknown)
+		{
+			TargetState.Access = Resource->GetTrackedState().Access;
+		}
+		
 		if (Resource->LoggingMode != ELoggingMode::None
 #if LOG_UNNAMED_RESOURCES
 			|| Resource->GetDebugName() == nullptr
 #endif
 			)
 		{
-			Log(Resource, SubresourceIndex, CreateTrace, TEXT("End"), TEXT("EndTransition"), *FString::Printf(TEXT("Access: %s, Pipeline: %s, Executing Pipeline: %s"),
+			const TCHAR* PulledFromTracked = TEXT("");
+
+			if (InTargetState.Access == ERHIAccess::Unknown)
+			{
+				PulledFromTracked = TEXT(" (Pulled From SetTrackedAccess)");
+			}
+
+			Log(Resource, SubresourceIndex, CreateTrace, TEXT("End"), TEXT("EndTransition"), *FString::Printf(TEXT("Access: %s%s, Pipeline: %s, Executing Pipeline: %s"),
 				*GetRHIAccessName(TargetState.Access),
+				PulledFromTracked,
 				*GetRHIPipelineName(TargetState.Pipelines),
 				*GetRHIPipelineName(ExecutingPipeline)
 			));
 		}
 
-		FPipelineState& State = States[ExecutingPipeline];
-
 		// Set bIgnoringAfterState to true in case of EResourceTransitionFlags::IgnoreAfterState so on the next transition we can check that EResourceTransitionFlags::IgnoreAfterState is used as well
-		if ( EnumHasAnyFlags(NewFlags, EResourceTransitionFlags::IgnoreAfterState))
+		if (EnumHasAnyFlags(NewFlags, EResourceTransitionFlags::IgnoreAfterState))
 		{
 			State.bIgnoringAfterState = !State.bIgnoringAfterState;
 		}
@@ -2437,7 +2542,7 @@ namespace RHIValidation
 		State.Current.Access = DecayResourceAccess(State.Current.Access, RequiredState.Access, bAllowAllUAVsOverlap || State.bExplicitAllowUAVOverlap);
 	}
 
-	void FSubresourceState::AssertTracked(FResource* Resource, FSubresourceIndex const& SubresourceIndex, const FState& RequiredState)
+	void FSubresourceState::AssertTracked(FResource* Resource, FSubresourceIndex const& SubresourceIndex, const FState& RequiredState, ERHIPipeline ExecutingPipeline)
 	{
 		if (Resource->LoggingMode != ELoggingMode::None
 #if LOG_UNNAMED_RESOURCES
@@ -2445,16 +2550,22 @@ namespace RHIValidation
 #endif
 			)
 		{
-			Log(Resource, SubresourceIndex, nullptr, nullptr, TEXT("AssertTracked"), *FString::Printf(TEXT("Access: %s, Pipeline %s"), *GetRHIAccessName(RequiredState.Access), *GetRHIPipelineName(RequiredState.Pipelines)));
+			Log(Resource, SubresourceIndex, nullptr, nullptr, TEXT("AssertTracked"), *FString::Printf(TEXT("Access: %s, Pipelines %s"), *GetRHIAccessName(RequiredState.Access), *GetRHIPipelineName(RequiredState.Pipelines)));
 		}
 
-		FPipelineState& State = States[RequiredState.Pipelines];
+		for (ERHIPipeline Pipeline : MakeFlagsRange(RequiredState.Pipelines))
+		{
+			FPipelineState& State = States[Pipeline];
 
-		// Check we're not trying to access the resource whilst a pending resource transition is in progress.
-		RHI_VALIDATION_CHECK(!State.bTransitioning, *GetReasonString_AccessDuringTransition(Resource, SubresourceIndex, State.Current, RequiredState, State.CreateTransitionBacktrace, State.BeginTransitionBacktrace));
+			// Check we're not trying to access the resource whilst a pending resource transition is in progress (can only do this on the executing pipeline).
+			if (State.Current.Pipelines == ExecutingPipeline)
+			{
+				RHI_VALIDATION_CHECK(!State.bTransitioning, *GetReasonString_AccessDuringTransition(Resource, SubresourceIndex, State.Current, RequiredState, State.CreateTransitionBacktrace, State.BeginTransitionBacktrace));
+			}
 
-		// Ensure the resource is in the required state for this operation
-		RHI_VALIDATION_CHECK(State.Current.Access == RequiredState.Access, *GetReasonString_IncorrectTrackedAccess(Resource, SubresourceIndex, State.Current, RequiredState));
+			// Ensure the resource is in the required state for this operation (ignore the Discard state which always resets).
+			RHI_VALIDATION_CHECK(State.Current.Access == ERHIAccess::Discard || State.Current == RequiredState, *GetReasonString_IncorrectSetTrackedAccess(Resource, SubresourceIndex, State.Current, RequiredState));
+		}
 	}
 
 	void FSubresourceState::SpecificUAVOverlap(FResource* Resource, FSubresourceIndex const& SubresourceIndex, ERHIPipeline Pipeline, bool bAllow)
@@ -2712,17 +2823,24 @@ namespace RHIValidation
 			break;
 
 		case EOpType::SetTrackedAccess:
+		{
 			Data_Assert.Identity.Resource->EnumerateSubresources(Data_SetTrackedAccess.Resource->GetWholeResourceRange(), [this, &Queue](FSubresourceState& State, FSubresourceIndex const& SubresourceIndex)
 			{
 				State.AssertTracked(
 					Data_SetTrackedAccess.Resource,
 					SubresourceIndex,
-					FState(Data_SetTrackedAccess.Access, Queue.Pipeline));
+					Data_SetTrackedAccess.State,
+					Queue.Pipeline);
 			});
-			Data_SetTrackedAccess.Resource->TrackedAccess = Data_SetTrackedAccess.Access;
-			Data_SetTrackedAccess.Resource->ReleaseOpRef();
+			FResource* Resource = Data_SetTrackedAccess.Resource;
+			if (Resource->LoggingMode != ELoggingMode::None)
+			{
+				Log(Resource, {}, nullptr, TEXT("SetTrackedAccess"), TEXT("SetTrackedAccess"), *FString::Printf(TEXT("Access: (%s), Pipelines: (%s), Executing Pipeline: (%s)"), *GetRHIAccessName(Data_SetTrackedAccess.State.Access), *GetRHIPipelineName(Data_SetTrackedAccess.State.Pipelines), *GetRHIPipelineName(Queue.Pipeline)) );
+			}
+			Resource->TrackedState = Data_SetTrackedAccess.State;
+			Resource->ReleaseOpRef();
 			break;
-
+		}
 		case EOpType::AcquireTransient:
 			Data_AcquireTransient.Resource->TransientState.Acquire(Data_AcquireTransient.Resource, Data_AcquireTransient.CreateBacktrace, Queue.Pipeline);
 			Data_AcquireTransient.Resource->ReleaseOpRef();
diff --git a/Engine/Source/Runtime/RHI/Public/RHIAccess.h b/Engine/Source/Runtime/RHI/Public/RHIAccess.h
--- a/Engine/Source/Runtime/RHI/Public/RHIAccess.h
+++ b/Engine/Source/Runtime/RHI/Public/RHIAccess.h
@@ -124,3 +124,6 @@ extern RHI_API ERHIAccess GRHIMergeableAccessMask;
 
 /** Mask of states which are allowed to be considered for multi-pipeline state merging. This should be a subset of GRHIMergeableAccessMask. */
 extern RHI_API ERHIAccess GRHIMultiPipelineMergeableAccessMask;
+
+/** [NeedsTransientDiscardStateTracking] State to use as an intermediate state when transitioning to discard and the resource has differing prior states per subresource. */
+extern RHI_API ERHIAccess GRHIMultiSubresourceDiscardIntermediateAccess;
\ No newline at end of file
diff --git a/Engine/Source/Runtime/RHI/Public/RHICommandList.h b/Engine/Source/Runtime/RHI/Public/RHICommandList.h
--- a/Engine/Source/Runtime/RHI/Public/RHICommandList.h
+++ b/Engine/Source/Runtime/RHI/Public/RHICommandList.h
@@ -456,6 +456,16 @@ public:
 	~FRHICommandListScopedPipelineGuard();
 };
 
+class FRHICommandListScopedAllowExtraTransitions
+{
+	FRHICommandListBase& RHICmdList;
+	bool bAllowExtraTransitions = false;
+
+public:
+	FRHICommandListScopedAllowExtraTransitions(FRHICommandListBase& RHICmdList, bool bAllowExtraTransitions);
+	~FRHICommandListScopedAllowExtraTransitions();
+};
+
 class FRHICommandListBase
 {
 protected:
@@ -497,12 +507,6 @@ public:
 			return false;
 		}
 
-		// SetTrackedAccess mutates the FRHIViewableResource::TrackedAccess member.
-		if (bUsesSetTrackedAccess)
-		{
-			return false;
-		}
-
 		// Some shader bundle implementations do not currently support parallel translate
 		if (bUsesShaderBundles && !GRHISupportsShaderBundleParallel)
 		{
@@ -945,16 +949,9 @@ PRAGMA_ENABLE_DEPRECATION_WARNINGS
 			FRHITextureCreateDesc NewCreateDesc(CreateDesc);
 			NewCreateDesc.SetInitialState(RHIGetDefaultResourceState(CreateDesc.Flags, CreateDesc.BulkData != nullptr));
 
-#if RHI_TRANSITIONS_DEBUG_INITIAL_STATE
-			UE_LOG(LogRHI, Display, TEXT("Resource: %s Creation State %s"), CreateDesc.DebugName, *GetRHIAccessName(NewCreateDesc.InitialState));
-#endif
-
 			return GDynamicRHI->RHICreateTexture(*this, NewCreateDesc);
 		}
 
-#if RHI_TRANSITIONS_DEBUG_INITIAL_STATE
-		UE_LOG(LogRHI, Display, TEXT("Resource: %s Creation State %s"), CreateDesc.DebugName, *GetRHIAccessName(CreateDesc.InitialState));
-#endif
 		return GDynamicRHI->RHICreateTexture(*this, CreateDesc);
 	}
 
@@ -1234,16 +1231,17 @@ PRAGMA_ENABLE_DEPRECATION_WARNINGS
 		return GRHIGlobals.NeedsShaderUnbinds;
 	}
 
-	// Returns true if the RHI needs and allow extra transitions
+	// Returns true if the underlying RHI needs implicit transitions inside of certain methods.
 	bool NeedsExtraTransitions() const
 	{
-		return GRHIGlobals.NeedExtraTransitions && bAllowExtraTransition;
+		return GRHIGlobals.NeedsExtraTransitions && bAllowExtraTransitions;
 	}
-	// Returns old state of bAllowExtraTransition
-	bool SetAllowExtraTransition(bool NewState)
+
+	// Returns old state of bAllowExtraTransitions
+	bool SetAllowExtraTransitions(bool NewState)
 	{
-		bool OldState = bAllowExtraTransition;
-		bAllowExtraTransition = NewState;
+		bool OldState = bAllowExtraTransitions;
+		bAllowExtraTransitions = NewState;
 		return OldState;
 	}
 
@@ -1359,7 +1357,7 @@ protected:
 	bool bUsesSetTrackedAccess   = false;
 	bool bUsesShaderBundles      = false;
 	bool bUsesLockFence          = false;
-	bool bAllowExtraTransition	 = true;
+	bool bAllowExtraTransitions  = true;
 
 	// The currently selected pipelines that RHI commands are directed to, during command list recording.
 	// This is also adjusted during command list execution based on recorded use of ActivatePipeline().
@@ -1553,11 +1551,6 @@ public:
 		}
 	}
 
-	FORCEINLINE_DEBUGGABLE bool NeedTransitionsPatching() const { return (PersistentState.bAllowResourceStateTracking) && (GRHIGlobals.EnableHighLevelTransitionPatching); }
-
-	// Returns true if the RHI always needs the state of both planes (depth|stencil) for transitions
-	FORCEINLINE_DEBUGGABLE bool NeedsAllPlanes() const { return GRHIGlobals.NeedsAllPlanes && GRHIGlobals.EnableHighLevelTransitionPatching; }
-
 private:
 	FRHICommandListBase(FPersistentState const& InPersistentState);
 
@@ -1583,62 +1576,6 @@ private:
 #if WITH_RHI_BREADCRUMBS
 	friend bool IRHIComputeContext::ShouldEmitBreadcrumbs() const;
 #endif
-
-	struct FTransitionPatchInfo
-	{
-		FTransitionPatchInfo(FRHITransitionInfo* InTransitionInfo, ERHIAccess InState) :
-			TransitionInfo(InTransitionInfo),
-			State(InState) {}
-
-		FRHITransitionInfo* TransitionInfo;
-		ERHIAccess			State;
-	};
-
-	struct FRenderPassPatchInfo
-	{
-		FRenderPassPatchInfo(FRHIRenderPassInfo* InRenderPassInfo, ERHIAccess InState) :
-			RenderPassInfo(InRenderPassInfo),
-			State(InState) {}
-
-		FRHIRenderPassInfo* RenderPassInfo;
-		ERHIAccess			State;
-	};
-
-	// Using FRHICmdListBaseArrayAllocator since elements Lifetime for all the following containers are 1 frame, between Record and Submit
-	// Track every resource state change happening during the command list Record
-	TMap<FRHIViewableResource*, TRHIResourceState<FRHICmdListBaseArrayAllocator>, FRHICmdListBaseSetAllocator> TrackedResources;
-
-	// Transitions with Unknown or Discard (when needed) that will patched at Dispatch time
-	TArray<FTransitionPatchInfo, FRHICmdListBaseArrayAllocator> TransitionInfoToPatch;
-	// (NeedsAllPlanes) Transitions with a single plane of DepthStencil specified where the other plane needs to be resolved with tracking
-	TArray<FTransitionPatchInfo, FRHICmdListBaseArrayAllocator> PlaneTransitionInfoToPatch;
-	// (NeedsAllPlanes) RenderPassInfo with a single plane of DepthStencil specified where the other plane needs to be resolved with tracking
-	TArray<FRenderPassPatchInfo, FRHICmdListBaseArrayAllocator> RenderPassInfoToPatch;
-	// (NeedsAllPlanes) Track transitions happening over multiple command lists
-	static TMap<FRHIViewableResource*, TArray<FRHITransitionInfo>> GlobalActiveTransitions;
-
-	// This will scan the array and keep track of transitions that needs patching later at Dispatch time
-	RHI_API void TrackResourceStateBegin(TArrayView<const FRHITransition*> Transitions);
-	RHI_API void TrackResourceStateEnd(TArrayView<const FRHITransition*> Transitions);
-
-	// Called at Dispatch time, in order of submission, it's when the Unknown and Discard transitions can be patched with the real state
-	void UpdateTransitionsAndResourceState();
-	
-	// Bypass needs to patch Transitions on the spot
-	inline void PatchTransitionsInPlace(FRHITransitionInfo* TransitionInfo, const ERHIPipeline SrcPipeline, const ERHIPipeline DestPipeline, bool bForceTrackingUpdate=false);
-	inline void PatchTransitionsInPlace(TArrayView<const FRHITransition*> Transitions);
-	inline void UpdateTrackedAccessInPlace(TArrayView<const FRHITransition*> Transitions);
-
-	// (NeedsAllPlanes) Resource needs tracking while it's transitioning in case a RenderPass queries its state between Begin and End
-	RHI_API static bool TrackTransitioningResource(const FRHITransitionCreateInfoData& TransitionCreateInfoData, const FRHITransitionInfo& TransitionInfo);
-	// (NeedsAllPlanes) Track transitions that happen over multiple cmdlists for resources flagged by TrackTransitioningResource
-	using ActiveTransitionContainer = TArray<const FRHITransitionInfo*, TInlineAllocator<2>>;
-	TMap<FRHIViewableResource*, ActiveTransitionContainer, FRHICmdListBaseSetAllocator> ActiveTransitions;
-	TMap<FRHIViewableResource*, ActiveTransitionContainer, FRHICmdListBaseSetAllocator> ClosedTransitions;
-
-protected:
-	inline void ResolveNopAccessInPlace(FRHIRenderPassInfo& InInfo);
-	inline void TrackNopAccess(const FRHIRenderPassInfo& InInfo);
 };
 
 #if WITH_RHI_BREADCRUMBS
@@ -3077,21 +3014,6 @@ public:
 #endif
 		if (Bypass())
 		{
-			if (NeedTransitionsPatching())
-			{
-				PatchTransitionsInPlace(Transitions);
-			}
-
-			//Create the Platform Transition for each FRHITransition
-			if (GRHIGlobals.EnableHighLevelTransitionPatching)
-			{
-				for (const FRHITransition* Transition : Transitions)
-				{
-					const bool bCleanupInfoData = !NeedsAllPlanes();
-					Transition->CreatePlatformTransition(bCleanupInfoData);
-				}
-			}
-
 			GetComputeContext().RHIBeginTransitions(Transitions);
 
 			for (const FRHITransition* Transition : Transitions)
@@ -3101,8 +3023,6 @@ public:
 		}
 		else
 		{
-			TrackResourceStateBegin(Transitions);
-
 			// Copy the transition array into the command list
 			FRHITransition** DstTransitionArray = (FRHITransition**)Alloc(sizeof(FRHITransition*) * Transitions.Num(), alignof(FRHITransition*));
 			FMemory::Memcpy(DstTransitionArray, Transitions.GetData(), sizeof(FRHITransition*) * Transitions.Num());
@@ -3117,15 +3037,6 @@ public:
 #endif
 		if (Bypass())
 		{
-
-#if UE_BUILD_DEBUG
-			for (const FRHITransition* Transition : Transitions)
-			{
-				checkSlow(Transition->IsPlatformLevelCreated());
-			}
-#endif
-			UpdateTrackedAccessInPlace(Transitions);
-
 			GetComputeContext().RHIEndTransitions(Transitions);
 
 			for (const FRHITransition* Transition : Transitions)
@@ -3135,8 +3046,6 @@ public:
 		}
 		else
 		{
-			TrackResourceStateEnd(Transitions);
-
 			// Copy the transition array into the command list
 			FRHITransition** DstTransitionArray = (FRHITransition**)Alloc(sizeof(FRHITransition*) * Transitions.Num(), alignof(FRHITransition*));
 			FMemory::Memcpy(DstTransitionArray, Transitions.GetData(), sizeof(FRHITransition*) * Transitions.Num());
@@ -3170,13 +3079,6 @@ public:
 
 	FORCEINLINE_DEBUGGABLE void SetTrackedAccess(TArrayView<const FRHITrackedAccessInfo> Infos)
 	{
-		if (NeedTransitionsPatching())
-		{
-			return;
-		}
-
-		bUsesSetTrackedAccess = true;
-
 		if (Bypass())
 		{
 			for (const FRHITrackedAccessInfo& Info : Infos)
@@ -3187,10 +3089,11 @@ public:
 		else
 		{
 			ALLOC_COMMAND(FRHICommandSetTrackedAccess)(AllocArray(Infos));
+			RHIThreadFence(true);
 		}
 	}
 
-	FORCEINLINE_DEBUGGABLE void SetTrackedAccess(TArrayView<const FRHITransitionInfo> Infos)
+	FORCEINLINE_DEBUGGABLE void SetTrackedAccess(TArrayView<const FRHITransitionInfo> Infos, ERHIPipeline PipelinesAfter)
 	{
 		FRHITrackedAccessInfo* TrackedAccessInfos = reinterpret_cast<FRHITrackedAccessInfo*>(Alloc(sizeof(FRHITrackedAccessInfo) * Infos.Num(), alignof(FRHITrackedAccessInfo)));
 		int32 NumTrackedAccessInfos = 0;
@@ -3201,7 +3104,7 @@ public:
 
 			if (FRHIViewableResource* Resource = GetViewableResource(Info))
 			{
-				new (&TrackedAccessInfos[NumTrackedAccessInfos++]) FRHITrackedAccessInfo(Resource, Info.AccessAfter);
+				new (&TrackedAccessInfos[NumTrackedAccessInfos++]) FRHITrackedAccessInfo(Resource, Info.AccessAfter, PipelinesAfter);
 			}
 		}
 
@@ -3598,7 +3501,7 @@ public:
 		}
 
 		SetGPUMask(PrevGPUMask);
-		
+
 		if (Bypass())
 		{
 			GetComputeContext().RHITransferResources(Params);
@@ -3882,11 +3785,6 @@ public:
 
 	FORCEINLINE_DEBUGGABLE void BuildAccelerationStructures(TConstArrayView<FRayTracingGeometryBuildParams> Params, const FRHIBufferRange& ScratchBufferRange)
 	{
-		if (NeedsExtraTransitions())
-		{
-			Transition(FRHITransitionInfo(ScratchBufferRange.Buffer, ERHIAccess::Unknown, ERHIAccess::UAVCompute));
-		}
-
 		if (Bypass())
 		{
 			GetComputeContext().RHIBuildAccelerationStructures(Params, ScratchBufferRange);
@@ -4426,20 +4324,8 @@ public:
 
 		InInfo.Validate();
 
-		// If the ExclusiveDepthStencil contains a Nop, RHIs that need the state of both planes 
-		// will need tracking to resolve the RHIAccess for the missing plane
-		const bool bNeedsNopResolve = NeedsAllPlanes() &&
-			InInfo.DepthStencilRenderTarget.DepthStencilTarget &&
-			IsStencilFormat(InInfo.DepthStencilRenderTarget.DepthStencilTarget->GetDesc().Format) &&
-			(!InInfo.DepthStencilRenderTarget.ExclusiveDepthStencil.IsUsingDepth() ||
-			 !InInfo.DepthStencilRenderTarget.ExclusiveDepthStencil.IsUsingStencil());
-
 		if (Bypass())
 		{
-			if (bNeedsNopResolve)
-			{
-				ResolveNopAccessInPlace(const_cast<FRHIRenderPassInfo&>(InInfo));
-			}
 			GetContext().RHIBeginRenderPass(InInfo, Name);
 		}
 		else
@@ -4448,11 +4334,6 @@ public:
 			FRHIRenderPassInfo* InfoCopy = (FRHIRenderPassInfo*)Alloc(sizeof(FRHIRenderPassInfo), alignof(FRHIRenderPassInfo));
 			FMemory::Memcpy(InfoCopy, &InInfo, sizeof(InInfo));
 
-			if (bNeedsNopResolve)
-			{
-				TrackNopAccess(*InfoCopy);
-			}
-
 			TCHAR* NameCopy  = AllocString(Name);
 			ALLOC_COMMAND(FRHICommandBeginRenderPass)(*InfoCopy, NameCopy);
 		}
diff --git a/Engine/Source/Runtime/RHI/Public/RHICommandList.inl b/Engine/Source/Runtime/RHI/Public/RHICommandList.inl
--- a/Engine/Source/Runtime/RHI/Public/RHICommandList.inl
+++ b/Engine/Source/Runtime/RHI/Public/RHICommandList.inl
@@ -147,6 +147,17 @@ inline FRHICommandListScopedPipelineGuard::~FRHICommandListScopedPipelineGuard()
 	}
 }
 
+inline FRHICommandListScopedAllowExtraTransitions::FRHICommandListScopedAllowExtraTransitions(FRHICommandListBase& RHICmdList, bool bInAllowExtraTransitions)
+	: RHICmdList(RHICmdList)
+{
+	bAllowExtraTransitions = RHICmdList.SetAllowExtraTransitions(bInAllowExtraTransitions);
+}
+
+inline FRHICommandListScopedAllowExtraTransitions::~FRHICommandListScopedAllowExtraTransitions()
+{
+	RHICmdList.SetAllowExtraTransitions(bAllowExtraTransitions);
+}
+
 inline FRHIResourceReplaceBatcher::~FRHIResourceReplaceBatcher()
 {
 	RHICmdList.ReplaceResources(MoveTemp(Infos));
@@ -272,222 +283,4 @@ inline void TRHILambdaCommandMultiPipe<RHICmdListType, LAMBDA>::ExecuteAndDestru
 	// Static cast to enforce const type in lambda args
 	Lambda(static_cast<FRHIContextArray const&>(Contexts));
 	Lambda.~LAMBDA();
-}
-
-FORCEINLINE_DEBUGGABLE void FRHICommandListBase::UpdateTrackedAccessInPlace(TArrayView<const FRHITransition*> Transitions)
-{
-	check(Bypass());
-
-	if (!NeedsAllPlanes())
-	{
-		return;
-	}
-
-	for (const FRHITransition* Transition : Transitions)
-	{
-		FRHITransitionCreateInfoData* TransitionCreateInfoData = Transition->CreationInfoData;
-		if (TransitionCreateInfoData)
-		{
-			const ERHIPipeline SrcPipeline = TransitionCreateInfoData->SrcPipelines;
-			const ERHIPipeline DestPipeline = TransitionCreateInfoData->DstPipelines;
-
-			if (SrcPipeline != DestPipeline)
-			{
-				continue;
-			}
-
-			for (int32 i = 0; i < TransitionCreateInfoData->TransitionInfos.Num(); i++)
-			{
-				const FRHITransitionInfo& TransitionInfo = TransitionCreateInfoData->TransitionInfos[i];
-
-				if ((TransitionInfo.Type == FRHITransitionInfo::EType::Buffer) && (GRHIGlobals.SupportImplicitBufferStateTracking))
-				{
-					// This RHI supports Buffers state tracking at platform level, so we don't need to track them
-					continue;
-				}
-
-				FRHIViewableResource* TransitionViewableResource = GetViewableResource(TransitionInfo);
-				if (TransitionViewableResource)
-				{
-					const bool bShouldSkipAfterState = ((EnumHasAnyFlags(TransitionInfo.Flags, EResourceTransitionFlags::IgnoreAfterState)) || (TransitionInfo.AccessAfter == ERHIAccess::Unknown));
-					const bool bNeedUpdateAfterState = (bShouldSkipAfterState == false);
-
-					if (bNeedUpdateAfterState)
-					{
-						TransitionViewableResource->TrackedAccess.SetSubresourceState(TransitionInfo, TransitionInfo.AccessAfter);
-						TransitionViewableResource->TrackedAccess.CheckAllSubresourceSame();
-					}
-				}
-			}
-
-			Transition->CleanupInfoData();
-		}
-	}
-}
-
-FORCEINLINE_DEBUGGABLE void FRHICommandListBase::PatchTransitionsInPlace(FRHITransitionInfo* TransitionInfo, const ERHIPipeline SrcPipeline, const ERHIPipeline DestPipeline, bool bForceTrackingUpdate)
-{
-	if ((TransitionInfo->Type == FRHITransitionInfo::EType::Buffer) && (GRHIGlobals.SupportImplicitBufferStateTracking))
-	{
-		// This RHI supports Buffers state tracking at platform level, so we don't need to track them
-		return;
-	}
-
-	const FRHISubresourceRange SubResourceIndex = *TransitionInfo;
-	FRHIViewableResource* TransitionViewableResource = GetViewableResource(*TransitionInfo);
-
-	// Both Unknown is not a supported case
-	check(!((TransitionInfo->AccessBefore == ERHIAccess::Unknown) && (TransitionInfo->AccessAfter == ERHIAccess::Unknown)));
-	// Discard to Unknown is not a supported case
-	check(!((TransitionInfo->AccessBefore == ERHIAccess::Discard) && (TransitionInfo->AccessAfter == ERHIAccess::Unknown)));
-
-	if (TransitionViewableResource)
-	{
-		// Some RHI will not transition to initial state on After = Discard when running on the Async queue
-		bool bDiscardOnAsync = false;
-		if (NeedsExtraTransitions())
-		{
-			if ((TransitionInfo->AccessAfter == ERHIAccess::Discard)
-				&& ((SrcPipeline == ERHIPipeline::AsyncCompute) || (SrcPipeline == ERHIPipeline::All))
-				&& (DestPipeline == ERHIPipeline::AsyncCompute))
-			{
-				bDiscardOnAsync = true;
-			}
-		}
-		const bool bShouldSkipAfterState = ((EnumHasAnyFlags(TransitionInfo->Flags, EResourceTransitionFlags::IgnoreAfterState)) || (TransitionInfo->AccessAfter == ERHIAccess::Unknown));
-		const bool bNeedUpdateAfterState = (bShouldSkipAfterState == false) && (bDiscardOnAsync == false);
-		
-		// Some RHI needs extra information on transition from Discard
-		const bool bBeforeIsDiscard = NeedsExtraTransitions() && (TransitionInfo->AccessBefore == ERHIAccess::Discard);
-
-		if (bBeforeIsDiscard)
-		{
-			TransitionInfo->DiscardAccessBefore = TransitionViewableResource->TrackedAccess.GetSubresourceState(SubResourceIndex);
-		}
-		else if (TransitionInfo->AccessBefore == ERHIAccess::Unknown)
-		{
-			TransitionInfo->AccessBefore = TransitionViewableResource->TrackedAccess.GetSubresourceState(SubResourceIndex);
-		}
-		else if (TransitionInfo->AccessAfter == ERHIAccess::Unknown)
-		{
-			TransitionInfo->AccessAfter = TransitionViewableResource->TrackedAccess.GetSubresourceState(SubResourceIndex);
-		}
-
-		// When NeedsAllPlanes is active, we might have a RenderPass querying TrackedAccess between the Begin/End if we remain on the Graphics pipe
-		const bool bCanUpdate = (bForceTrackingUpdate || !NeedsAllPlanes() || 
-			!EnumHasAnyFlags(SrcPipeline, ERHIPipeline::Graphics) || 
-			!EnumHasAnyFlags(DestPipeline, ERHIPipeline::Graphics));
-		if (bNeedUpdateAfterState && bCanUpdate)
-		{
-			TransitionViewableResource->TrackedAccess.SetSubresourceState(SubResourceIndex, TransitionInfo->AccessAfter);
-			TransitionViewableResource->TrackedAccess.CheckAllSubresourceSame();
-		}
-		
-		check(TransitionInfo->AccessBefore != ERHIAccess::Unknown);
-		check(TransitionInfo->AccessAfter != ERHIAccess::Unknown);
-
-		// Some RHI need all planes of a DepthStencil texture to transition
-		if (NeedsAllPlanes() && (TransitionInfo->Type == FRHITransitionInfo::EType::Texture) && !SubResourceIndex.IsAllPlaneSlices())
-		{
-			// Only applies to texture formats with a stencil
-			FRHITexture* Texture = (FRHITexture*)TransitionViewableResource;
-			if (IsStencilFormat(TransitionInfo->Texture->GetDesc().Format))
-			{
-				// Set the OtherPlaneState for all single plane transitions, less work than tracking the set planes per subresource
-				const FRHISubresourceRange OtherResourceIndex(
-					SubResourceIndex.MipIndex,
-					SubResourceIndex.ArraySlice,
-					(SubResourceIndex.PlaneSlice == FRHISubresourceRange::kDepthPlaneSlice) ? FRHISubresourceRange::kStencilPlaneSlice : FRHISubresourceRange::kDepthPlaneSlice
-				);
-
-				TransitionInfo->OtherPlaneAccess = TransitionViewableResource->TrackedAccess.GetSubresourceState(OtherResourceIndex);
-				TransitionInfo->Flags |= EResourceTransitionFlags::FillOtherPlaneState;
-			}
-		}
-	}
-}
-
-
-FORCEINLINE_DEBUGGABLE void FRHICommandListBase::PatchTransitionsInPlace(TArrayView<const FRHITransition*> Transitions)
-{
-	check(Bypass());
-
-	for (const FRHITransition* Transition : Transitions)
-	{
-		FRHITransitionCreateInfoData* TransitionCreateInfoData = Transition->CreationInfoData;
-		if (TransitionCreateInfoData)
-		{
-			for (int32 i = 0; i < TransitionCreateInfoData->TransitionInfos.Num(); i++)
-			{
-				PatchTransitionsInPlace((FRHITransitionInfo*)&TransitionCreateInfoData->TransitionInfos[i], TransitionCreateInfoData->SrcPipelines, TransitionCreateInfoData->DstPipelines);
-			}
-		}
-	}
-}
-
-FORCEINLINE_DEBUGGABLE void FRHICommandListBase::ResolveNopAccessInPlace(FRHIRenderPassInfo& InInfo)
-{
-	if (!InInfo.DepthStencilRenderTarget.ExclusiveDepthStencil.IsUsingDepth())
-	{
-		const FRHISubresourceRange SubResourceIndex(0, 0, FRHISubresourceRange::kDepthPlaneSlice);
-		InInfo.NopAccess = InInfo.DepthStencilRenderTarget.DepthStencilTarget->TrackedAccess.GetSubresourceState(SubResourceIndex);
-	}
-	else if (!InInfo.DepthStencilRenderTarget.ExclusiveDepthStencil.IsUsingStencil())
-	{
-		const FRHISubresourceRange SubResourceIndex(0, 0, FRHISubresourceRange::kStencilPlaneSlice);
-		InInfo.NopAccess = InInfo.DepthStencilRenderTarget.DepthStencilTarget->TrackedAccess.GetSubresourceState(SubResourceIndex);
-	}
-}
-
-FORCEINLINE_DEBUGGABLE void FRHICommandListBase::TrackNopAccess(const FRHIRenderPassInfo& InInfo)
-{
-	checkSlow(NeedsAllPlanes());
-	if (InInfo.DepthStencilRenderTarget.DepthStencilTarget)
-	{
-		ERHIAccess NopState = ERHIAccess::Unknown;
-
-		if (ActiveTransitionContainer* TransitionArray = ActiveTransitions.Find(InInfo.DepthStencilRenderTarget.DepthStencilTarget))
-		{
-			for (const FRHITransitionInfo* TransitionInfo : *TransitionArray)
-			{
-				if (!InInfo.DepthStencilRenderTarget.ExclusiveDepthStencil.IsUsingDepth())
-				{
-					if (TransitionInfo->IsAllPlaneSlices() || TransitionInfo->PlaneSlice == FRHISubresourceRange::kDepthPlaneSlice)
-					{
-						NopState = TransitionInfo->AccessBefore;
-						break;
-					}
-				}
-				else if (!InInfo.DepthStencilRenderTarget.ExclusiveDepthStencil.IsUsingStencil())
-				{
-					if (TransitionInfo->IsAllPlaneSlices() || TransitionInfo->PlaneSlice == FRHISubresourceRange::kStencilPlaneSlice)
-					{
-						NopState = TransitionInfo->AccessBefore;
-						break;
-					}
-				}
-			}
-		}
-
-		if (NopState == ERHIAccess::Unknown)
-		{
-			TRHIResourceState<FRHICmdListBaseArrayAllocator>* TrackedResourceState = TrackedResources.Find(InInfo.DepthStencilRenderTarget.DepthStencilTarget);
-			if (TrackedResourceState)
-			{
-				if (!InInfo.DepthStencilRenderTarget.ExclusiveDepthStencil.IsUsingDepth())
-				{
-					const FRHISubresourceRange SubResourceIndex(0, 0, FRHISubresourceRange::kDepthPlaneSlice);
-					NopState = TrackedResourceState->GetSubresourceState(SubResourceIndex);
-				}
-				else if (!InInfo.DepthStencilRenderTarget.ExclusiveDepthStencil.IsUsingStencil())
-				{
-					const FRHISubresourceRange SubResourceIndex(0, 0, FRHISubresourceRange::kStencilPlaneSlice);
-					NopState = TrackedResourceState->GetSubresourceState(SubResourceIndex);
-				}
-			}
-		}
-
-		// Note: even if the state is a valid state we cannot write into FRHIRenderPassInfo since we are recording the cmdlist
-		RenderPassInfoToPatch.Push(FRenderPassPatchInfo((FRHIRenderPassInfo*)&InInfo, NopState));
-	}
-}
+}
\ No newline at end of file
diff --git a/Engine/Source/Runtime/RHI/Public/RHICommandListCommandExecutes.inl b/Engine/Source/Runtime/RHI/Public/RHICommandListCommandExecutes.inl
--- a/Engine/Source/Runtime/RHI/Public/RHICommandListCommandExecutes.inl
+++ b/Engine/Source/Runtime/RHI/Public/RHICommandListCommandExecutes.inl
@@ -383,17 +383,6 @@ void FRHICommandBeginTransitions::Execute(FRHICommandListBase& CmdList)
 {
 	RHISTAT(BeginTransitions);
 
-	// Create the platform level transitions
-	if (GRHIGlobals.EnableHighLevelTransitionPatching)
-	{
-		for (const FRHITransition* Transition : Transitions)
-		{
-			const bool bCleanupInfoData = !CmdList.NeedsAllPlanes();
-			Transition->CreatePlatformTransitionThreadSafe(bCleanupInfoData);
-			checkSlow(Transition->IsPlatformLevelCreated());
-		}
-	}
-
 	INTERNAL_DECORATOR_COMPUTE(RHIBeginTransitions)(Transitions);
 
 	for (const FRHITransition* Transition : Transitions)
@@ -406,17 +395,6 @@ void FRHICommandEndTransitions::Execute(FRHICommandListBase& CmdList)
 {
 	RHISTAT(EndTransitions);
 
-	// Create the platform level transitions
-	if (GRHIGlobals.EnableHighLevelTransitionPatching)
-	{
-		for (const FRHITransition* Transition : Transitions)
-		{
-			const bool bCleanupInfoData = !CmdList.NeedsAllPlanes();
-			Transition->CreatePlatformTransitionThreadSafe(bCleanupInfoData);
-			checkSlow(Transition->IsPlatformLevelCreated());
-		}
-	}
-
 	INTERNAL_DECORATOR_COMPUTE(RHIEndTransitions)(Transitions);
 
 	for (const FRHITransition* Transition : Transitions)
@@ -429,15 +407,6 @@ void FRHICommandResourceTransition::Execute(FRHICommandListBase& CmdList)
 {
 	RHISTAT(ResourceTransition);
 
-	// Create FRHITransition from FRHITransitionHighLevel with Platform information
-	if (GRHIGlobals.EnableHighLevelTransitionPatching)
-	{
-		const bool bCleanupInfoData = true;
-		Transition->CreatePlatformTransitionThreadSafe(bCleanupInfoData);
-	}
-
-	checkSlow(Transition->IsPlatformLevelCreated());
-
 	INTERNAL_DECORATOR_COMPUTE(RHIBeginTransitions)(MakeArrayView((const FRHITransition**)&Transition, 1));
 	INTERNAL_DECORATOR_COMPUTE(RHIEndTransitions)(MakeArrayView((const FRHITransition**)&Transition, 1));
 
diff --git a/Engine/Source/Runtime/RHI/Public/RHIContext.h b/Engine/Source/Runtime/RHI/Public/RHIContext.h
--- a/Engine/Source/Runtime/RHI/Public/RHIContext.h
+++ b/Engine/Source/Runtime/RHI/Public/RHIContext.h
@@ -535,15 +535,19 @@ public:
 #endif
 	void SetTrackedAccess(const FRHITrackedAccessInfo& Info)
 	{
-		check(Info.Resource != nullptr);
-		check(Info.Access != ERHIAccess::Unknown);
-		Info.Resource->SetTrackedAccessFromContext(Info.Access);
+		Info.Resource->SetTrackedAccessFromContext(Info);
 	}
 
 	inline ERHIAccess GetTrackedAccess(const FRHIViewableResource* Resource) const
 	{
 		check(Resource);
-		return Resource->TrackedAccess.GetResourceState();
+		return Resource->TrackedAccess.Access;
+	}
+
+	inline ERHIPipeline GetTrackedPipelines(const FRHIViewableResource* Resource) const
+	{
+		check(Resource);
+		return Resource->TrackedAccess.Pipelines != ERHIPipeline::None ? Resource->TrackedAccess.Pipelines : GetPipeline();
 	}
 
 	virtual void* RHIGetNativeCommandBuffer() { return nullptr; }
diff --git a/Engine/Source/Runtime/RHI/Public/RHIDefinitions.h b/Engine/Source/Runtime/RHI/Public/RHIDefinitions.h
--- a/Engine/Source/Runtime/RHI/Public/RHIDefinitions.h
+++ b/Engine/Source/Runtime/RHI/Public/RHIDefinitions.h
@@ -1399,6 +1399,11 @@ enum class ERHITransitionCreateFlags
 
 	// Indicates that the transition is allowed to happen during a RenderPass
 	AllowDuringRenderPass = 1 << 2,
+
+	// This is an advanced mode that allows issuing a transition using a single pipeline as the previous pipelines in a transition
+	// for a resource that was previously transitioned to ERHIPipeline::All. It assumes that a manual fence was issued to synchronize
+	// the prior async compute work. It also assumes that tracked access for a resource (using SetTrackedAccess) is up to date.
+	AllowDecayPipelines = 1 << 3
 };
 
 ENUM_CLASS_FLAGS(ERHITransitionCreateFlags);
@@ -1408,14 +1413,9 @@ enum class EResourceTransitionFlags : uint8
 	None                = 0,
 
 	MaintainCompression = 1 << 0, // Specifies that the transition should not decompress the resource, allowing us to read a compressed resource directly in its compressed state.
-	Discard				= 1 << 1, // Specifies that the data in the resource should be discarded during the transition - used for transient resource acquire when the resource will be fully overwritten
-	Clear				= 1 << 2, // Specifies that the data in the resource should be cleared during the transition - used for transient resource acquire when the resource might not be fully overwritten
-
-	// Flags below aren't used by RDG and therefore not included in the Last/Mask calculations
-	IgnoreAfterState	= 1 << 3, // The transition will be patched for Unknown if needed but the After State will be ignored by the tracker
-	FillOtherPlaneState = 1 << 4, // The other plane state of a depth|stencil texture has been filled by tracking
+	IgnoreAfterState    = 1 << 1, // The transition will be patched for Unknown if needed but the After State will be ignored by the tracker
 
-	Last = Clear,
+	Last = IgnoreAfterState,
 	Mask = (Last << 1) - 1
 };
 ENUM_CLASS_FLAGS(EResourceTransitionFlags);
@@ -1664,8 +1664,4 @@ struct FShaderCodeValidationUBSize
 {
 	uint16 BindPoint;
 	uint32 Size;
-};
-
-#define RHI_TRANSITIONS_DEBUG_INITIAL_STATE	0
-
-#define RHI_UnknownSubResourceCount (0)
\ No newline at end of file
+};
\ No newline at end of file
diff --git a/Engine/Source/Runtime/RHI/Public/RHIGlobals.h b/Engine/Source/Runtime/RHI/Public/RHIGlobals.h
--- a/Engine/Source/Runtime/RHI/Public/RHIGlobals.h
+++ b/Engine/Source/Runtime/RHI/Public/RHIGlobals.h
@@ -722,21 +722,15 @@ struct FRHIGlobals
 	/** true if the RHI supports shader execution reordering. */
 	bool SupportsShaderExecutionReordering = false;
 
-	/** True if RHI can automatically handle resource transition for Buffers type, therefore tracking at high level is not needed*/
-	bool SupportImplicitBufferStateTracking = false;
-
-	/** True if RHI ask for patching of Unknown transition states in the RHI high level */
-	bool EnableHighLevelTransitionPatching = false;
-
 	/** True if RHI needs the transitions from the high level for COPYSRC/COPYDEST*/
-	bool NeedExtraTransitions = false;
+	bool NeedsExtraTransitions = false;
+	
+	/** True if RHI needs to track states when transitioning between Discard. */
+	bool NeedsTransientDiscardStateTracking = false;
 
 	/** True if RHI is unable to service transitioning to ERHIAccess::Discard from all pipes to async compute. */
 	bool NeedsTransientDiscardOnGraphicsWorkaround = false;
 
-	/** True if RHI always needs the state of both planes (depth|stencil) for transitions */
-	bool NeedsAllPlanes = false;
-
 	/** True if lossy/fixed-rate compression of render targets is supporterd. */
 	bool SupportsLossyFramebufferCompression = false;
 
diff --git a/Engine/Source/Runtime/RHI/Public/RHIResources.h b/Engine/Source/Runtime/RHI/Public/RHIResources.h
--- a/Engine/Source/Runtime/RHI/Public/RHIResources.h
+++ b/Engine/Source/Runtime/RHI/Public/RHIResources.h
@@ -1325,12 +1325,7 @@ class FRHIViewableResource : public FRHIResource
 public:
 	ERHIAccess GetTrackedAccess_Unsafe() const
 	{
-		return TrackedAccess.GetResourceState();
-	}
-
-	bool AreAllSubresourcesSame_Unsafe() const
-	{
-		return TrackedAccess.AreAllSubresourcesSame();
+		return TrackedAccess.Access;
 	}
 
 	FName GetName() const
@@ -1338,26 +1333,20 @@ public:
 		return Name;
 	}
 
-	/** @return The number of SubResources in the TrackedState. */
-	uint16 GetTrackedSubResourceCount() const
-	{
-		return TrackedAccess.SubResourceCount;
-	}
-
 #if ENABLE_RHI_VALIDATION
 	virtual RHIValidation::FResource* GetValidationTrackerResource() = 0;
 #endif
 
 protected:
-	FRHIViewableResource(ERHIResourceType InResourceType, ERHIAccess InAccess, uint16 SubResourceCount)
+	FRHIViewableResource(ERHIResourceType InResourceType, ERHIAccess InAccess)
 		: FRHIResource(InResourceType)
-		, TrackedAccess(InAccess, SubResourceCount)
+		, TrackedAccess(InAccess)
 	{
 	}
 
-	virtual void SetTrackedAccessFromContext(ERHIAccess Access)
+	virtual void SetTrackedAccessFromContext(FRHITrackedAccess InTrackedAccess)
 	{
-		TrackedAccess = Access;
+		TrackedAccess = InTrackedAccess;
 	}
 
 	void TakeOwnership(FRHIViewableResource& Other)
@@ -1373,7 +1362,7 @@ protected:
 	FName Name;
 
 private:
-	TRHIResourceState<FDefaultAllocator> TrackedAccess;
+	FRHITrackedAccess TrackedAccess;
 
 	friend class FRHICommandListBase;
 	friend class FRHIComputeCommandList;
@@ -1645,7 +1634,7 @@ class FRHIBuffer : public FRHIViewableResource
 public:
 	/** Initialization constructor. */
 	FRHIBuffer(const FRHIBufferCreateDesc& InCreateDesc)
-		: FRHIViewableResource(RRT_Buffer, InCreateDesc.InitialState, 1)
+		: FRHIViewableResource(RRT_Buffer, InCreateDesc.InitialState)
 		, Desc(InCreateDesc)
 	{
 #if RHI_ENABLE_RESOURCE_INFO
@@ -2357,7 +2346,7 @@ private:
 	friend class FRHITextureReference;
 	/** Constructor for texture references */
 	FRHITexture(ERHIResourceType InResourceType)
-		: FRHIViewableResource(InResourceType, ERHIAccess::Unknown, RHI_UnknownSubResourceCount)
+		: FRHIViewableResource(InResourceType, ERHIAccess::Unknown)
 	{
 		check(InResourceType == RRT_TextureReference);
 	}
diff --git a/Engine/Source/Runtime/RHI/Public/RHITextureReference.h b/Engine/Source/Runtime/RHI/Public/RHITextureReference.h
--- a/Engine/Source/Runtime/RHI/Public/RHITextureReference.h
+++ b/Engine/Source/Runtime/RHI/Public/RHITextureReference.h
@@ -42,7 +42,7 @@ protected:
 	friend class FDynamicRHI;
 
 	// Only called from Platform Context
-	RHI_API virtual void SetTrackedAccessFromContext(ERHIAccess Access) override;
+	RHI_API virtual void SetTrackedAccessFromContext(FRHITrackedAccess TrackedAccess) override;
 
 	// Called only from FDynamicRHI::RHIUpdateTextureReference
 	void SetReferencedTexture(FRHITexture* InTexture)
diff --git a/Engine/Source/Runtime/RHI/Public/RHITransition.h b/Engine/Source/Runtime/RHI/Public/RHITransition.h
--- a/Engine/Source/Runtime/RHI/Public/RHITransition.h
+++ b/Engine/Source/Runtime/RHI/Public/RHITransition.h
@@ -138,8 +138,6 @@ struct FRHITransitionInfo : public FRHISubresourceRange
 
 	ERHIAccess AccessBefore = ERHIAccess::Unknown;
 	ERHIAccess AccessAfter = ERHIAccess::Unknown;
-	ERHIAccess DiscardAccessBefore = ERHIAccess::Unknown;
-	ERHIAccess OtherPlaneAccess = ERHIAccess::Unknown;  // used with GRHIGlobals.NeedsAllPlanes
 	EResourceTransitionFlags Flags = EResourceTransitionFlags::None;
 	TOptional<FRHICommitResourceInfo> CommitInfo;
 
@@ -216,7 +214,6 @@ struct FRHITransitionInfo : public FRHISubresourceRange
 		return Resource == RHS.Resource
 			&& Type == RHS.Type
 			&& AccessBefore == RHS.AccessBefore
-			&& DiscardAccessBefore == RHS.DiscardAccessBefore
 			&& AccessAfter == RHS.AccessAfter
 			&& Flags == RHS.Flags
 			&& FRHISubresourceRange::operator==(RHS);
@@ -228,7 +225,6 @@ struct FRHITransitionInfo : public FRHISubresourceRange
 	}
 };
 
-
 struct FRHITransientAliasingOverlap
 {
 	union
@@ -382,297 +378,6 @@ struct FRHITransientAliasingInfo
 	}
 };
 
-//==================================================================================================================================
-// FRHIResourceState
-// Tracking of per-resource or per-subresource state
-//==================================================================================================================================
-template <typename TSecondaryAllocator>
-class TRHIResourceState
-{
-public:
-	struct FSubResourceElement
-	{
-		FSubResourceElement() {}
-		FSubResourceElement(const FRHISubresourceRange& Range, ERHIAccess State)
-		: SubRange(Range)
-		, SubState(State)
-		{}
-
-		FRHISubresourceRange SubRange;
-		ERHIAccess SubState;
-	};
-
-	explicit TRHIResourceState(uint16 SubResourceCount) 
-		: ResourceState(ERHIAccess::Unknown)
-		, UnsetSubResourceState(ERHIAccess::Unknown)
-		, SubResourceCount(SubResourceCount)
-		, bAllSubresourcesSame(1)
-		, bCmdListTracking(0)
-	{}
-
-	explicit TRHIResourceState(ERHIAccess inState, uint16 SubResourceCount, bool CmdListTracking = false) 
-		: ResourceState(inState)
-		, UnsetSubResourceState(ERHIAccess::Unknown)
-		, SubResourceCount(SubResourceCount)
-		, bAllSubresourcesSame(1)
-		, bCmdListTracking(CmdListTracking)
-	{}
-
-	explicit TRHIResourceState(const FRHISubresourceRange& Range, ERHIAccess State, uint16 SubResourceCount, bool CmdListTracking = false)
-		: ResourceState(ERHIAccess::Unknown)
-		, UnsetSubResourceState(ERHIAccess::Unknown)
-		, SubResourceCount(SubResourceCount)
-		, bAllSubresourcesSame(1)
-		, bCmdListTracking(CmdListTracking)
-	{
-		SetSubresourceState(Range, State);
-	}
-
-	~TRHIResourceState() {}
-
-	friend class TRHIResourceState<FRHICmdListBaseArrayAllocator>;
-	
-	//Using for set "global" resource state from cmdlist states
-	inline void Merge(const TRHIResourceState<FRHICmdListBaseArrayAllocator>& Other)
-	{
-		check(Other.bCmdListTracking);
-
-		if (Other.bAllSubresourcesSame)
-		{ 
-			check(Other.ResourceState != ERHIAccess::Unknown)
-			check(Other.SubResourcesState.IsEmpty() == true);
-			SetResourceState(Other.ResourceState);
-		}
-		else
-		{
-			for (int32 i = 0; i < Other.SubResourcesState.Num(); i++)
-			{
-				SetSubresourceState(Other.SubResourcesState[i].SubRange, Other.SubResourcesState[i].SubState);
-			}
-		}
-
-		//Compaction of subresource tracking in to resource tracking, might need to move it at explicit calls to avoid perf repercussions
-		if (bAllSubresourcesSame == false)
-		{
-			if (SubResourcesState.Num() == SubResourceCount)
-			{
-				ERHIAccess AllSubResState = SubResourcesState[0].SubState;
-
-				for (int32 i = 1; i < SubResourcesState.Num(); i++)
-				{
-					if (SubResourcesState[i].SubState != AllSubResState)
-					{
-						AllSubResState = ERHIAccess::Unknown;
-						break;
-					}
-				}
-
-				if (AllSubResState != ERHIAccess::Unknown)
-				{
-					SetResourceState(AllSubResState);
-				}
-			}
-		}
-	}
-
-	inline bool AreAllSubresourcesSame() const { return bAllSubresourcesSame && (ResourceState != ERHIAccess::Unknown); }
-	inline bool CheckAllSubresourceSame()
-	{
-		// already marked same?
-		if (bAllSubresourcesSame)
-		{
-			return true;
-		}
-		
-		if (SubResourcesState.Num() < SubResourceCount)
-		{
-			return false;
-		}
-
-		ERHIAccess State = SubResourcesState[0].SubState;
-
-		// All subresources must be individually checked
-		for (int32 i = 1; i < SubResourcesState.Num(); i++)
-		{
-			if (SubResourcesState[i].SubState != State)
-			{
-				return false;
-			}
-		}
-
-		SetResourceState(State);
-
-		return true;
-	}
-
-	inline ERHIAccess GetResourceState() const
-	{
-		return GetSubresourceState(FRHISubresourceRange(FRHISubresourceRange::kAllSubresources, FRHISubresourceRange::kAllSubresources, FRHISubresourceRange::kAllSubresources));
-	}
-	
-	inline ERHIAccess GetSubresourceState(const FRHISubresourceRange& SubresourceIndex) const
-	{
-		checkSlow(SubResourceCount != RHI_UnknownSubResourceCount);
-
-		if (bAllSubresourcesSame) 
-		{
-			return ResourceState;
-		}
-		else
-		{
-			if (SubresourceIndex.IsWholeResource())
-			{
-				//trying to get the status of the whole resource while subresource are in different state each one
-				if (SubResourcesState.Num() == SubResourceCount)
-				{
-					ERHIAccess AllSubResState = SubResourcesState[0].SubState;
-
-					for (int32 i = 1; i < SubResourcesState.Num(); i++)
-					{
-						if (SubResourcesState[i].SubState != AllSubResState)
-						{
-							AllSubResState = ERHIAccess::Unknown;
-							break;
-						}
-					}
-
-					check(AllSubResState != ERHIAccess::Unknown);
-					return AllSubResState;
-				}
-
-				// Edge case, only resource 0 is set and we are asking for the whole resource state
-				if ((SubResourcesState.Num() == 1) && (SubResourcesState[0].SubRange.IsZeroIndexResource()))
-				{
-					return SubResourcesState[0].SubState;
-				}
-
-				// This should not happen, we are asking for the whole resource state but subresources are in different state within each other
-				check(false);
-				return ERHIAccess::Unknown;
-			}
-
-			for (int32 i = 0; i < SubResourcesState.Num(); i++)
-			{
-				if (SubResourcesState[i].SubRange == SubresourceIndex)
-				{
-					return SubResourcesState[i].SubState;
-				}
-			}
-
-			//trying to get the status of one subresources that we didn't set yet, return the Unset value
-			if (SubResourcesState.Num() < SubResourceCount)
-			{
-				return UnsetSubResourceState;
-			}
-
-			//resource wasn't there 
-			return ERHIAccess::Unknown;
-		}
-	}
-
-	inline void SetResourceState(ERHIAccess State)
-	{
-		SetSubresourceState(FRHISubresourceRange(FRHISubresourceRange::kAllSubresources, FRHISubresourceRange::kAllSubresources, FRHISubresourceRange::kAllSubresources), State);
-	}
-	
-	inline void SetSubresourceState(const FRHISubresourceRange& SubresourceIndex, ERHIAccess State)
-	{
-		checkSlow(SubResourceCount != RHI_UnknownSubResourceCount);
-		checkSlow(IsValidAccess(State));
-
-		// If setting all subresources
-		if ((SubresourceIndex.IsWholeResource()) || (SubResourceCount == 1))
-		{
-			if (SubresourceIndex.IsWholeResource() == false)
-			{
-				check(SubresourceIndex.IsZeroIndexResource());
-			}
-
-			bAllSubresourcesSame = 1;
-			ResourceState = State;
-			SubResourcesState.Empty(SubResourceCount);
-		}
-		else
-		{
-			check(SubresourceIndex.MipIndex * SubresourceIndex.ArraySlice * SubresourceIndex.PlaneSlice <= SubResourceCount);
-			
-			SubResourcesState.Reserve(SubResourceCount);
-			
-			// If state was previously tracked on a per-resource level, then transition to per-subresource tracking
-			if (bAllSubresourcesSame)
-			{
-				for (int32 i = 0; i < SubResourcesState.Num(); i++)
-				{
-					SubResourcesState[i].SubState = ResourceState;
-				}
-
-				bAllSubresourcesSame = 0;
-
-				// State is now tracked per-subresource, m_ResourceState still would be the not set subresources state
-				UnsetSubResourceState = ResourceState;
-				ResourceState = ERHIAccess::Unknown;
-			}
-
-			bool bFound = false;
-			for (int32 i = 0; i < SubResourcesState.Num(); i++)
-			{
-				if (SubResourcesState[i].SubRange == SubresourceIndex)
-				{
-					bFound = true;
-					SubResourcesState[i].SubState = State;
-					break;
-				}
-			}
-
-			if (bFound == false)
-			{
-				SubResourcesState.Push(FSubResourceElement(SubresourceIndex, State));
-				check(SubResourcesState.Num() <= SubResourceCount);
-			}
-
-		}
-	}
-
-	//Operators
-	inline TRHIResourceState& operator= (ERHIAccess State)
-	{
-		check(bAllSubresourcesSame == 1);
-		SetResourceState(State);
-		return *this;
-	}
-
-	inline bool operator == (const ERHIAccess Other) const
-	{
-		check(bAllSubresourcesSame == 1);
-		return Other == ResourceState;
-	}
-
-	bool IsValid()
-	{
-		if (bAllSubresourcesSame == false)
-		{
-			if (SubResourcesState.IsEmpty())
-				return false;
-		}
-
-		return true;
-	}
-
-	// SubresourceState[0...n] status of the subresources if any
-	TArray<FSubResourceElement, TSizedInlineAllocator<8, 32, TSecondaryAllocator> > SubResourcesState; //todo test if a sparse container would perform better
-	// ResourceState is the status of the whole resource
-	ERHIAccess ResourceState;
-	// Status of unset SubResources
-	ERHIAccess UnsetSubResourceState;
-	// Total number of SubResources
-	uint16 SubResourceCount;
-	// Set to 1 if ResourceState is valid.  In this case, all subresources have the same state
-	// Set to 0 if SubresourceState is valid.  In this case, each subresources may have a different state (or may be unknown)
-	uint8 bAllSubresourcesSame : 1;
-	// Set to 1 if the State is used on Command list tracking  
-	uint8 bCmdListTracking : 1;
-};
-
 struct FRHITransitionCreateInfo
 {
 	FRHITransitionCreateInfo() = default;
@@ -724,90 +429,42 @@ struct FRHITransitionCreateInfo
 
 };
 
-struct FRHITrackedAccessInfo
+struct FRHITrackedAccess
+{
+	FRHITrackedAccess() = default;
+
+	FRHITrackedAccess(ERHIAccess InAccess)
+		: Access(InAccess)
+	{}
+
+	FRHITrackedAccess(ERHIAccess InAccess, ERHIPipeline InPipelines)
+		: Access(InAccess)
+		, Pipelines(InPipelines)
+	{}
+
+	ERHIAccess Access = ERHIAccess::Unknown;
+	ERHIPipeline Pipelines = ERHIPipeline::None;
+};
+
+struct FRHITrackedAccessInfo : FRHITrackedAccess
 {
 	FRHITrackedAccessInfo() = default;
 
-	FRHITrackedAccessInfo(FRHIViewableResource* InResource, ERHIAccess InAccess)
-		: Resource(InResource)
-		, Access(InAccess)
+	FRHITrackedAccessInfo(FRHIViewableResource* InResource, ERHIAccess InAccess, ERHIPipeline InPipelines)
+		: FRHITrackedAccess(InAccess, InPipelines)
+		, Resource(InResource)
 	{}
 
 	FRHIViewableResource* Resource = nullptr;
-	ERHIAccess Access = ERHIAccess::Unknown;
 };
 
 struct FRHITransition;
 RHI_API FRHITransition* RHICreateTransition(const FRHITransitionCreateInfo& CreateInfo);
 
-
 enum class ETransitionFlag : uint8
 {
 	None = 0,
 	AllowDuringRenderPass = 1 << 0,			// Specific transition are happening during RenderPass and are reported by the validation layer, we need to allow some of them
-	PlatformLevelCreated  = 1 << 1			// Platform level has been created on this Transition
-};
-
-// Struct used to store creation info and create the opaque platform dependent FRHITransition after tracking and patching
-struct FRHITransitionCreateInfoData
-{
-private:
-	// Prevent copying and moving. Only pointers to these structures are allowed.
-	FRHITransitionCreateInfoData(const FRHITransitionCreateInfoData&) = delete;
-	FRHITransitionCreateInfoData(FRHITransitionCreateInfoData&&) = delete;
-
-public:
-	FRHITransitionCreateInfoData(const FRHITransitionCreateInfo& CreationInfo)
-	{
-		SrcPipelines = CreationInfo.SrcPipelines;
-		DstPipelines = CreationInfo.DstPipelines;
-		m_Flags = CreationInfo.Flags;
-
-		TransitionInfos = CreationInfo.TransitionInfos;
-		AliasingInfos = CreationInfo.AliasingInfos;
-
-		// Deep copy of AlisianInfo used for validation
-		if (GRHIValidationEnabled)
-		{
-			uint32 OffsetDataIndex = 0;
-			uint32 AllOverlapNum = 0;
-			for (FRHITransientAliasingInfo& AliasingInfo : AliasingInfos)
-			{
-				AllOverlapNum += AliasingInfo.Overlaps.Num();
-			}
-
-			TransientAliasingOverlapData.Reserve(AllOverlapNum);
-			if (AllOverlapNum)
-			{
-				for (FRHITransientAliasingInfo& AliasingInfo : AliasingInfos)
-				{
-					if (!AliasingInfo.Overlaps.IsEmpty())
-					{
-						TransientAliasingOverlapData.Append(AliasingInfo.Overlaps.GetData(), AliasingInfo.Overlaps.Num());
-						AliasingInfo.Overlaps = TArrayView<const FRHITransientAliasingOverlap>(&TransientAliasingOverlapData[OffsetDataIndex], AliasingInfo.Overlaps.Num());
-						OffsetDataIndex += AliasingInfo.Overlaps.Num();
-					}
-				}
-			}
-		}
-	}
-
-	virtual ~FRHITransitionCreateInfoData()
-	{}
-
-	inline const FRHITransitionCreateInfo GetCreationInfo() const
-	{
-		return FRHITransitionCreateInfo(SrcPipelines, DstPipelines, m_Flags, MakeArrayView(TransitionInfos), MakeArrayView(AliasingInfos));
-	}
-
-	ERHIPipeline SrcPipelines = ERHIPipeline::None;
-	ERHIPipeline DstPipelines = ERHIPipeline::None;
-	mutable ERHITransitionCreateFlags m_Flags = ERHITransitionCreateFlags::None;
-	TArray<FRHITransitionInfo, FRHICmdListBaseArrayAllocator>		 TransitionInfos;
-	TArray<FRHITransientAliasingInfo, FRHICmdListBaseArrayAllocator> AliasingInfos;
-
-	// We need to copy aliasing overlap information 
-	TArray<FRHITransientAliasingOverlap, FRHICmdListBaseArrayAllocator> TransientAliasingOverlapData;
 };
 
 // Opaque data structure used to represent a pending resource transition in the RHI.
@@ -832,9 +489,7 @@ public:
 		return const_cast<FRHITransition*>(this)->GetPrivateData<T>();
 	}
 
-	inline bool AllowInRenderingPass() const { return EnumHasAnyFlags(Flags, ETransitionFlag::AllowDuringRenderPass); }
-
-	inline bool IsPlatformLevelCreated() const { return EnumHasAnyFlags(Flags, ETransitionFlag::PlatformLevelCreated); }
+	inline bool AllowInRenderingPass() const { return EnumHasAnyFlags(Flags, ERHITransitionCreateFlags::AllowDuringRenderPass); }
 
 private:
 	// Prevent copying and moving. Only pointers to these structures are allowed.
@@ -842,23 +497,15 @@ private:
 	FRHITransition(FRHITransition&&) = delete;
 
 	// Private constructor. Memory for transitions is allocated manually with extra space at the tail of the structure for RHI use.
-	FRHITransition(ERHIPipeline SrcPipelines, ERHIPipeline DstPipelines, ETransitionFlag NewFlags = ETransitionFlag::None)
+	FRHITransition(ERHIPipeline SrcPipelines, ERHIPipeline DstPipelines, ERHITransitionCreateFlags InFlags)
 		: State(int8(int32(SrcPipelines) | (int32(DstPipelines) << int32(ERHIPipeline::Num))))
-		, Flags(NewFlags)
+		, Flags(InFlags)
 #if DO_CHECK || USING_CODE_ANALYSIS
 		, AllowedSrc(SrcPipelines)
 		, AllowedDst(DstPipelines)
 #endif
 	{}
 
-	~FRHITransition()
-	{
-		if (CreationInfoData)
-		{
-			CleanupInfoData();
-		}
-	}
-
 	// Give private access to specific functions/RHI commands that need to allocate or control transitions.
 	friend RHI_API FRHITransition* RHICreateTransition(const FRHITransitionCreateInfo& CreateInfo);
 	friend class FRHICommandListBase;
@@ -877,9 +524,6 @@ private:
 	{
 		return FMath::Max((uint64)alignof(FRHITransition), GRHITransitionPrivateData_AlignInBytes);
 	}
-	
-	RHI_API void CreatePlatformTransition(bool bCleanupInfoData) const;
-	RHI_API void CreatePlatformTransitionThreadSafe(bool bCleanupInfoData) const;
 
 	inline void MarkBegin(ERHIPipeline Pipeline) const
 	{
@@ -913,17 +557,11 @@ private:
 		}
 	}
 	
-	RHI_API void CleanupInfoData() const;
-
 	RHI_API void Cleanup() const;
 
 	mutable int8 State;
 	static_assert((int32(ERHIPipeline::Num) * 2) < (sizeof(State) * 8), "Not enough bits to hold pipeline state.");
-	mutable ETransitionFlag Flags;
-
-	// Deep copy of FRHITransitionCreateInfo used to track resource state and patch Unknown transition
-	mutable FRHITransitionCreateInfoData* CreationInfoData = nullptr;
-	mutable UE::FMutex TransitionCreationMutex;
+	mutable ERHITransitionCreateFlags Flags;
 
 #if DO_CHECK || USING_CODE_ANALYSIS
 	mutable ERHIPipeline AllowedSrc;
@@ -931,7 +569,4 @@ private:
 #endif
 };
 
-
-RHI_API FRHIViewableResource* GetViewableResource(const FRHITransitionInfo& Info);
-
-
+RHI_API FRHIViewableResource* GetViewableResource(const FRHITransitionInfo& Info);
\ No newline at end of file
diff --git a/Engine/Source/Runtime/RHI/Public/RHIValidationCommon.h b/Engine/Source/Runtime/RHI/Public/RHIValidationCommon.h
--- a/Engine/Source/Runtime/RHI/Public/RHIValidationCommon.h
+++ b/Engine/Source/Runtime/RHI/Public/RHIValidationCommon.h
@@ -205,7 +205,7 @@ namespace RHIValidation
 		void BeginTransition   (FResource* Resource, FSubresourceIndex const& SubresourceIndex, const FState& CurrentStateFromRHI, const FState& TargetState, EResourceTransitionFlags NewFlags, ERHITransitionCreateFlags CreateFlags, ERHIPipeline Pipeline, const TRHIPipelineArray<uint64>& PipelineMaxAwaitedFenceValues, void* CreateTrace);
 		void EndTransition     (FResource* Resource, FSubresourceIndex const& SubresourceIndex, const FState& CurrentStateFromRHI, const FState& TargetState, EResourceTransitionFlags NewFlags, ERHIPipeline Pipeline, uint64 PipelineFenceValue, void* CreateTrace);
 		void Assert            (FResource* Resource, FSubresourceIndex const& SubresourceIndex, const FState& RequiredState, bool bAllowAllUAVsOverlap);
-		void AssertTracked     (FResource* Resource, FSubresourceIndex const& SubresourceIndex, const FState& TrackedState);
+		void AssertTracked     (FResource* Resource, FSubresourceIndex const& SubresourceIndex, const FState& RequiredState, ERHIPipeline ExecutingPipeline);
 		void SpecificUAVOverlap(FResource* Resource, FSubresourceIndex const& SubresourceIndex, ERHIPipeline Pipeline, bool bAllow);
 	};
 
@@ -332,7 +332,7 @@ namespace RHIValidation
 		uint32 NumArraySlices = 0;
 		uint32 NumPlanes = 0;
 		FTransientState TransientState;
-		ERHIAccess TrackedAccess = ERHIAccess::Unknown;
+		FState TrackedState{ ERHIAccess::Unknown, ERHIPipeline::None };
 
 	private:
 		FString DebugName;
@@ -368,9 +368,9 @@ namespace RHIValidation
 			check(RefCount >= 0);
 		}
 
-		inline ERHIAccess GetTrackedAccess() const
+		inline FState GetTrackedState() const
 		{
-			return TrackedAccess;
+			return TrackedState;
 		}
 
 		inline uint32 GetNumSubresources() const
@@ -400,50 +400,10 @@ namespace RHIValidation
 			return Identity;
 		}
 
-		inline void InitTransient(const TCHAR* InDebugName)
-		{
-			check(TransientState.bTransient && TransientState.Status != FTransientState::EStatus::Acquired);
-			TransientState.Status = FTransientState::EStatus::None;
-			DebugName = InDebugName;
-			TrackedAccess = ERHIAccess::Discard;
-
-			for (ERHIPipeline Pipeline : MakeFlagsRange(ERHIPipeline::All))
-			{
-				auto& State = WholeResourceState.States[Pipeline];
-
-				State.Current.Access = ERHIAccess::Discard;
-				State.Current.Pipelines = Pipeline;
-				State.Previous = State.Current;
-			}
-			SubresourceStates.Reset();
-		}
+		void InitTransient(const TCHAR* InDebugName);
 
 	protected:
-		inline void InitBarrierTracking(int32 InNumMips, int32 InNumArraySlices, int32 InNumPlanes, ERHIAccess InResourceState, const TCHAR* InDebugName)
-		{
-			checkSlow(InNumMips > 0 && InNumArraySlices > 0 && InNumPlanes > 0);
-			check(InResourceState != ERHIAccess::Unknown);
-
-			NumMips = InNumMips;
-			NumArraySlices = InNumArraySlices;
-			NumPlanes = InNumPlanes;
-			TransientState = FTransientState(InResourceState);
-			TrackedAccess = InResourceState;
-
-			for (ERHIPipeline Pipeline : MakeFlagsRange(ERHIPipeline::All))
-			{
-				auto& State = WholeResourceState.States[Pipeline];
-
-				State.Current.Access = InResourceState;
-				State.Current.Pipelines = Pipeline;
-				State.Previous = State.Current;
-			}
-
-			if (InDebugName != nullptr)
-			{
-				SetDebugName(InDebugName);
-			}
-		}
+		void InitBarrierTracking(int32 InNumMips, int32 InNumArraySlices, int32 InNumPlanes, ERHIAccess InResourceState, const TCHAR* InDebugName);
 	};
 
 	inline bool FSubresourceRange::IsWholeResource(FResource& Resource) const
@@ -670,7 +630,7 @@ namespace RHIValidation
 			struct
 			{
 				FResource* Resource;
-				ERHIAccess Access;
+				FState State;
 			} Data_SetTrackedAccess;
 
 			struct
@@ -777,14 +737,14 @@ namespace RHIValidation
 			return MoveTemp(Op);
 		}
 
-		static inline FOperation SetTrackedAccess(FResource* Resource, ERHIAccess Access)
+		static inline FOperation SetTrackedAccess(FResource* Resource, FState State)
 		{
 			Resource->AddOpRef();
 
 			FOperation Op;
 			Op.Type = EOpType::SetTrackedAccess;
 			Op.Data_SetTrackedAccess.Resource = Resource;
-			Op.Data_SetTrackedAccess.Access = Access;
+			Op.Data_SetTrackedAccess.State = State;
 			return MoveTemp(Op);
 		}
 
@@ -1044,9 +1004,9 @@ namespace RHIValidation
 		}
 #endif
 
-		void SetTrackedAccess(FResource* Resource, ERHIAccess Access)
+		void SetTrackedAccess(FResource* Resource, ERHIAccess Access, ERHIPipeline Pipelines)
 		{
-			AddOp(FOperation::SetTrackedAccess(Resource, Access));
+			AddOp(FOperation::SetTrackedAccess(Resource, FState(Access, Pipelines)));
 		}
 
 		void Rename(FResource* Resource, const TCHAR* NewName, const TCHAR* Suffix = nullptr)
diff --git a/Engine/Source/Runtime/RHI/Public/RHIValidationContext.h b/Engine/Source/Runtime/RHI/Public/RHIValidationContext.h
--- a/Engine/Source/Runtime/RHI/Public/RHIValidationContext.h
+++ b/Engine/Source/Runtime/RHI/Public/RHIValidationContext.h
@@ -125,8 +125,9 @@ public:
 	{
 		check(Info.Resource != nullptr);
 		check(Info.Access != ERHIAccess::Unknown);
+		check(Info.Pipelines != ERHIPipeline::None);
 
-		Tracker->SetTrackedAccess(Info.Resource->GetValidationTrackerResource(), Info.Access);
+		Tracker->SetTrackedAccess(Info.Resource->GetValidationTrackerResource(), Info.Access, Info.Pipelines);
 
 		RHIContext->SetTrackedAccess(Info);
 	}
@@ -726,7 +727,7 @@ public:
 		check(Info.Resource != nullptr);
 		check(Info.Access != ERHIAccess::Unknown);
 
-		Tracker->SetTrackedAccess(Info.Resource->GetValidationTrackerResource(), Info.Access);
+		Tracker->SetTrackedAccess(Info.Resource->GetValidationTrackerResource(), Info.Access, Info.Pipelines);
 
 		RHIContext->SetTrackedAccess(Info);
 	}
diff --git a/Engine/Source/Runtime/RHICore/Internal/RHICoreTransitions.h b/Engine/Source/Runtime/RHICore/Internal/RHICoreTransitions.h
new file mode 100644
--- /dev/null
+++ b/Engine/Source/Runtime/RHICore/Internal/RHICoreTransitions.h
@@ -0,0 +1,45 @@
+// Copyright Epic Games, Inc. All Rights Reserved.
+
+#pragma once
+
+#include "RHITransition.h"
+#include "RHIContext.h"
+
+namespace UE::RHICore
+{
+
+struct FResourceState
+{
+	FResourceState(IRHIComputeContext& Context, ERHIPipeline InSrcPipelines, ERHIPipeline InDstPipelines, const FRHITransitionInfo& Info)
+		: AccessBefore(Info.AccessBefore)
+		, AccessAfter(Info.AccessAfter)
+		, SrcPipelines(InSrcPipelines)
+		, DstPipelines(InDstPipelines)
+	{
+		if (FRHIViewableResource* ViewableResource = GetViewableResource(Info))
+		{
+			if (AccessBefore == ERHIAccess::Unknown || AccessBefore == ERHIAccess::Discard)
+			{
+				SrcPipelines = Context.GetTrackedPipelines(ViewableResource);
+				AccessBefore = Context.GetTrackedAccess(ViewableResource);
+			}
+
+			if (AccessAfter == ERHIAccess::Unknown)
+			{
+				AccessAfter = Context.GetTrackedAccess(ViewableResource);
+			}
+
+			check(AccessBefore != ERHIAccess::Unknown);
+			check(AccessAfter != ERHIAccess::Unknown);
+		}
+	}
+
+	ERHIAccess AccessBefore;
+	ERHIAccess AccessAfter;
+
+	// SrcPipeline may differ from the declared RHI SrcPipeline when CollapseToSinglePipeline is used, in which case this will be ERHIPipeline::All.
+	ERHIPipeline SrcPipelines;
+	ERHIPipeline DstPipelines;
+};
+
+} //! UE::RHICore
\ No newline at end of file
diff --git a/Engine/Source/Runtime/RenderCore/Private/RenderGraphBuilder.cpp b/Engine/Source/Runtime/RenderCore/Private/RenderGraphBuilder.cpp
--- a/Engine/Source/Runtime/RenderCore/Private/RenderGraphBuilder.cpp
+++ b/Engine/Source/Runtime/RenderCore/Private/RenderGraphBuilder.cpp
@@ -2677,13 +2677,13 @@ void FRDGBuilder::SubmitBufferUploads(FRHICommandListBase& RHICmdListUpload, UE:
 		AllocateUploadBuffersTask->Trigger();
 	}
 
-	if ( (RHICmdListUpload.NeedsExtraTransitions() && RHICmdListUpload.NeedTransitionsPatching()) && UploadedBuffers.Num() > 1)
+	if ( (RHICmdListUpload.NeedsExtraTransitions()) && UploadedBuffers.Num() > 1)
 	{
 		SCOPED_NAMED_EVENT_TEXT("Upload_Multiple", FColor::Magenta);
 		
 		// This is here because we are explicitly batching a series of transition for all the buffers and we don't the individual extra transitions in Lock/Unlock
-		bool PrevCmdListAllowState = RHICmdListUpload.SetAllowExtraTransition(false);
-		
+		FRHICommandListScopedAllowExtraTransitions ScopedExtraTransitions(RHICmdList, false);
+
 		TSet<FRHIBuffer*, DefaultKeyFuncs<void*>, FConcurrentLinearSetAllocator> BuffersSet;
 		TArray<FRHITransitionInfo, FConcurrentLinearArrayAllocator> CopyDestTransitionInfo;
 		TArray<FRHITransitionInfo, FConcurrentLinearArrayAllocator> RevertTransitionInfo;
@@ -2761,9 +2761,6 @@ void FRDGBuilder::SubmitBufferUploads(FRHICommandListBase& RHICmdListUpload, UE:
 		BuffersSet.Reset();
 		CopyDestTransitionInfo.Reset();
 		RevertTransitionInfo.Reset();
-
-		// Set the cmdlist state to the original state
-		RHICmdListUpload.SetAllowExtraTransition(PrevCmdListAllowState);
 	}
 	else
 	{
@@ -4031,7 +4028,7 @@ void FRDGBuilder::AddLastTextureTransition(FRDGTexture* Texture)
 	if (Texture->AccessModeState.ActiveMode == FRDGViewableResource::EAccessMode::External)
 	{
 		// Assign the final state that was enqueued by the external access pass, which may include merged states.
-		EpilogueResourceAccesses.Emplace(Texture->GetRHI(), Texture->State[0]->Access);
+		EpilogueResourceAccesses.Emplace(Texture->GetRHI(), Texture->State[0]->Access, Texture->State[0]->GetPipelines());
 		return;
 	}
 
@@ -4044,16 +4041,44 @@ void FRDGBuilder::AddLastTextureTransition(FRDGTexture* Texture)
 	// Texture is using the RHI transient allocator. Transition it back to Discard in the final pass it is used.
 	if (Texture->DiscardPass.IsValid())
 	{
-		SubresourceStateAfter.SetPass(GetPassPipeline(Texture->DiscardPass), Texture->DiscardPass);
+		const ERHIPipeline DiscardPassPipeline = GetPassPipeline(Texture->DiscardPass);
+
+		SubresourceStateAfter.SetPass(DiscardPassPipeline, Texture->DiscardPass);
 		SubresourceStateAfter.BarrierLocation = ERDGBarrierLocation::Epilogue;
 		SubresourceStateAfter.Access = ERHIAccess::Discard;
+
+		if (GRHIGlobals.NeedsTransientDiscardStateTracking)
+		{
+			ERHIAccess EpilogueAccess = ERHIAccess::Unknown;
+
+			// Edge Case: Discarding Texture with RTV | DSV and multiple differing subresource states on async compute. Since we can't put multiple states
+			// inside of the TrackedAccess we have to do an intermediate transition instead on discard by ORing the intermediate state with Discard. This
+			// is going to be an incredibly rare case but needs to be handled correctly nonetheless.
+			if (DiscardPassPipeline == ERHIPipeline::AsyncCompute && EnumHasAnyFlags(Texture->Desc.Flags, ETextureCreateFlags::RenderTargetable | ETextureCreateFlags::DepthStencilTargetable))
+			{
+				EpilogueAccess = Texture->State[0]->Access;
+
+				for (uint32 SubresourceIndex = 1; SubresourceIndex < Texture->GetSubresourceCount(); ++SubresourceIndex)
+				{
+					FRDGSubresourceState* SubresourceState = Texture->State[SubresourceIndex];
+
+					if (!SubresourceState || SubresourceState->Access != EpilogueAccess)
+					{
+						EpilogueAccess = GRHIMultiSubresourceDiscardIntermediateAccess;
+						SubresourceStateAfter.Access |= GRHIMultiSubresourceDiscardIntermediateAccess;
+						break;
+					}
+				}
+			}
+
+			EpilogueResourceAccesses.Emplace(Texture->GetRHI(), ERHIAccess::Discard | EpilogueAccess, DiscardPassPipeline);
+		}
 	}
 	else
 	{
 		SubresourceStateAfter.Access = Texture->EpilogueAccess;
 
-		// Transient resources stay in the Discard state.
-		EpilogueResourceAccesses.Emplace(Texture->GetRHI(), SubresourceStateAfter.Access);
+		EpilogueResourceAccesses.Emplace(Texture->GetRHI(), SubresourceStateAfter.Access, ERHIPipeline::Graphics);
 	}
 
 	// Transition any unused (null) sub-resources to the epilogue state since we are assigning a monolithic state across all subresources.
@@ -4063,9 +4088,15 @@ void FRDGBuilder::AddLastTextureTransition(FRDGTexture* Texture)
 		{
 			if (!SubresourceStateBefore)
 			{
-				const FRDGPassHandle AcquirePass = Texture->AcquirePass.IsValid() ? Texture->FirstPass : GetProloguePassHandle();
-
 				SubresourceStateBefore = AllocSubresource();
+
+				FRDGPassHandle AcquirePass = GetProloguePassHandle();
+
+				if (Texture->AcquirePass.IsValid())
+				{
+					AcquirePass = Texture->FirstPass;
+				}
+
 				SubresourceStateBefore->SetPass(GetPassPipeline(AcquirePass), AcquirePass);
 			}
 
@@ -4123,7 +4154,7 @@ void FRDGBuilder::AddLastBufferTransition(FRDGBuffer* Buffer)
 	if (Buffer->AccessModeState.IsExternalAccess())
 	{
 		// Assign the final state that was enqueued by the external access pass, which may include merged states.
-		EpilogueResourceAccesses.Emplace(Buffer->GetRHI(), Buffer->State->Access);
+		EpilogueResourceAccesses.Emplace(Buffer->GetRHI(), Buffer->State->Access, Buffer->State->GetPipelines());
 		return;
 	}
 
@@ -4134,9 +4165,16 @@ void FRDGBuilder::AddLastBufferTransition(FRDGBuffer* Buffer)
 	// Texture is using the RHI transient allocator. Transition it back to Discard in the final pass it is used.
 	if (Buffer->DiscardPass.IsValid())
 	{
-		StateAfter->SetPass(GetPassPipeline(Buffer->DiscardPass), Buffer->DiscardPass);
+		const ERHIPipeline DiscardPassPipeline = GetPassPipeline(Buffer->DiscardPass);
+
+		StateAfter->SetPass(DiscardPassPipeline, Buffer->DiscardPass);
 		StateAfter->BarrierLocation = ERDGBarrierLocation::Epilogue;
 		StateAfter->Access = ERHIAccess::Discard;
+
+		if (GRHIGlobals.NeedsTransientDiscardStateTracking)
+		{
+			EpilogueResourceAccesses.Emplace(Buffer->GetRHI(), ERHIAccess::Discard, DiscardPassPipeline);
+		}
 	}
 	else
 	{
@@ -4144,7 +4182,7 @@ void FRDGBuilder::AddLastBufferTransition(FRDGBuffer* Buffer)
 		StateAfter->Access = Buffer->EpilogueAccess;
 		StateAfter->ReservedCommitHandle = AcquireReservedCommitHandle(Buffer);
 
-		EpilogueResourceAccesses.Emplace(Buffer->GetRHI(), StateAfter->Access);
+		EpilogueResourceAccesses.Emplace(Buffer->GetRHI(), StateAfter->Access, StateAfter->GetPipelines());
 	}
 
 	AddBufferTransition(Buffer, Buffer->State, StateAfter);
@@ -4166,7 +4204,7 @@ void FRDGBuilder::AddCulledReservedCommitTransition(FRDGBufferRef Buffer)
 	StateAfter->Access = Buffer->EpilogueAccess;
 	StateAfter->ReservedCommitHandle = AcquireReservedCommitHandle(Buffer);
 
-	EpilogueResourceAccesses.Emplace(Buffer->GetRHI(), StateAfter->Access);
+	EpilogueResourceAccesses.Emplace(Buffer->GetRHI(), StateAfter->Access, StateAfter->GetPipelines());
 	Buffer->Allocation = nullptr;
 
 	AddBufferTransition(Buffer, StateBefore, StateAfter);
@@ -4269,11 +4307,6 @@ void FRDGBuilder::AddTextureTransition(FRDGTexture* Texture, FRDGTextureSubresou
 
 				EResourceTransitionFlags Flags = SubresourceStateAfter->Flags;
 
-				if (SubresourceStateBefore->Access == ERHIAccess::Discard)
-				{
-					Flags |= EResourceTransitionFlags::Discard;
-				}
-
 				FRDGTransitionInfo Info;
 				Info.AccessBefore            = (uint64)SubresourceStateBefore->Access;
 				Info.AccessAfter             = (uint64)SubresourceStateAfter->Access;
@@ -5139,4 +5172,4 @@ void FRDGBuilder::ForceCopyCrossGPU()
 		RHICmdList.TransferResources(Transfers);
 	}
 }
-#endif  // WITH_MGPU
\ No newline at end of file
+#endif  // WITH_MGPU
diff --git a/Engine/Source/Runtime/RenderCore/Private/RenderGraphPass.cpp b/Engine/Source/Runtime/RenderCore/Private/RenderGraphPass.cpp
--- a/Engine/Source/Runtime/RenderCore/Private/RenderGraphPass.cpp
+++ b/Engine/Source/Runtime/RenderCore/Private/RenderGraphPass.cpp
@@ -200,25 +200,25 @@ void FRDGDispatchPassBuilder::Finish()
 	RHICmdList->FinishRecording();
 }
 
-FRDGBarrierBatchBegin::FRDGBarrierBatchBegin(ERHIPipeline InPipelineToBegin, ERHIPipeline InPipelinesToEnd, const TCHAR* InDebugName, FRDGPass* InDebugPass)
+FRDGBarrierBatchBegin::FRDGBarrierBatchBegin(ERHIPipeline InPipelineToBegin, ERHIPipeline InPipelinesToEnd, const TCHAR* InName, FRDGPass* InPass)
 	: PipelinesToBegin(InPipelineToBegin)
 	, PipelinesToEnd(InPipelinesToEnd)
 #if RDG_ENABLE_DEBUG
 	, DebugPasses(InPlace, nullptr)
-	, DebugName(InDebugName)
+	, DebugName(InName)
 #endif
 {
 #if RDG_ENABLE_DEBUG
-	DebugPasses[InPipelineToBegin] = InDebugPass;
+	DebugPasses[InPipelineToBegin] = InPass;
 #endif
 }
 
-FRDGBarrierBatchBegin::FRDGBarrierBatchBegin(ERHIPipeline InPipelinesToBegin, ERHIPipeline InPipelinesToEnd, const TCHAR* InDebugName, FRDGPassesByPipeline InDebugPasses)
+FRDGBarrierBatchBegin::FRDGBarrierBatchBegin(ERHIPipeline InPipelinesToBegin, ERHIPipeline InPipelinesToEnd, const TCHAR* InName, FRDGPassesByPipeline InPasses)
 	: PipelinesToBegin(InPipelinesToBegin)
 	, PipelinesToEnd(InPipelinesToEnd)
 #if RDG_ENABLE_DEBUG
-	, DebugPasses(InDebugPasses)
-	, DebugName(InDebugName)
+	, DebugPasses(InPasses)
+	, DebugName(InName)
 #endif
 {}
 
diff --git a/Engine/Source/Runtime/RenderCore/Private/RenderGraphResourcePool.cpp b/Engine/Source/Runtime/RenderCore/Private/RenderGraphResourcePool.cpp
--- a/Engine/Source/Runtime/RenderCore/Private/RenderGraphResourcePool.cpp
+++ b/Engine/Source/Runtime/RenderCore/Private/RenderGraphResourcePool.cpp
@@ -372,6 +372,9 @@ void FRDGTransientResourceAllocator::ReleasePendingDeallocations()
 
 	if (!PendingDeallocationList.IsEmpty())
 	{
+		TArray<FRHITrackedAccessInfo, SceneRenderingAllocator> EpilogueResourceAccesses;
+		EpilogueResourceAccesses.Reserve(PendingDeallocationList.Num());
+
 		TArray<FRHITransitionInfo, SceneRenderingAllocator> Transitions;
 		Transitions.Reserve(PendingDeallocationList.Num());
 
@@ -384,6 +387,7 @@ void FRDGTransientResourceAllocator::ReleasePendingDeallocations()
 			Fences.SetGraphics(0);
 			Allocator->DeallocateMemory(RenderTarget->Texture, Fences);
 			Transitions.Emplace(RenderTarget->Texture->GetRHI(), ERHIAccess::Unknown, ERHIAccess::Discard);
+			EpilogueResourceAccesses.Emplace(RenderTarget->Texture->GetRHI(), ERHIAccess::Discard, ERHIPipeline::Graphics);
 
 			RenderTarget->Reset();
 			RenderTarget->LifetimeState = ERDGTransientResourceLifetimeState::Deallocated;
@@ -395,6 +399,7 @@ void FRDGTransientResourceAllocator::ReleasePendingDeallocations()
 			FRHICommandListImmediate& RHICmdList = FRHICommandListExecutor::GetImmediateCommandList();
 			RHICmdList.BeginTransition(Transition);
 			RHICmdList.EndTransition(Transition);
+			RHICmdList.SetTrackedAccess(EpilogueResourceAccesses);
 		}
 
 		FreeList.Append(PendingDeallocationList);
diff --git a/Engine/Source/Runtime/RenderCore/Public/RenderGraphPass.h b/Engine/Source/Runtime/RenderCore/Public/RenderGraphPass.h
--- a/Engine/Source/Runtime/RenderCore/Public/RenderGraphPass.h
+++ b/Engine/Source/Runtime/RenderCore/Public/RenderGraphPass.h
@@ -56,7 +56,7 @@ struct FRDGBarrierBatchBeginId
 
 struct FRDGTransitionInfo
 {
-	static_assert((int32)ERHIAccess::Last <= (1 << 20) && (int32)ERDGViewableResourceType::MAX <= 3 && (int32)EResourceTransitionFlags::Last == (1 << 2), "FRDGTransitionInfo packing is no longer correct.");
+	static_assert((int32)ERHIAccess::Last <= (1 << 20) && (int32)ERDGViewableResourceType::MAX <= 3 && (int32)EResourceTransitionFlags::Last <= (1 << 2), "FRDGTransitionInfo packing is no longer correct.");
 
 	uint64 AccessBefore            : 21; // 21
 	uint64 AccessAfter             : 21; // 42
@@ -107,8 +107,8 @@ struct FRDGBarrierBatchEndId
 class FRDGBarrierBatchBegin
 {
 public:
-	RENDERCORE_API FRDGBarrierBatchBegin(ERHIPipeline PipelinesToBegin, ERHIPipeline PipelinesToEnd, const TCHAR* DebugName, FRDGPass* DebugPass);
-	RENDERCORE_API FRDGBarrierBatchBegin(ERHIPipeline PipelinesToBegin, ERHIPipeline PipelinesToEnd, const TCHAR* DebugName, FRDGPassesByPipeline DebugPasses);
+	RENDERCORE_API FRDGBarrierBatchBegin(ERHIPipeline PipelinesToBegin, ERHIPipeline PipelinesToEnd, const TCHAR* Name, FRDGPass* Pass);
+	RENDERCORE_API FRDGBarrierBatchBegin(ERHIPipeline PipelinesToBegin, ERHIPipeline PipelinesToEnd, const TCHAR* Name, FRDGPassesByPipeline Passes);
 
 	RENDERCORE_API void AddTransition(FRDGViewableResource* Resource, FRDGTransitionInfo Info);
 
@@ -122,7 +122,7 @@ public:
 		}
 		else
 		{
-			TransitionFlags = ERHITransitionCreateFlags::None;
+			EnumRemoveFlags(TransitionFlags, ERHITransitionCreateFlags::NoFence);
 		}
 		bTransitionNeeded = true;
 	}
@@ -149,7 +149,7 @@ private:
 	TRHIPipelineArray<FRDGBarrierBatchEndId> BarriersToEnd;
 	TArray<FRDGTransitionInfo, FRDGArrayAllocator> Transitions;
 	TArray<FRHITransientAliasingInfo, FRDGArrayAllocator> Aliases;
-	ERHITransitionCreateFlags TransitionFlags = ERHITransitionCreateFlags::NoFence;
+	ERHITransitionCreateFlags TransitionFlags = ERHITransitionCreateFlags::NoFence | ERHITransitionCreateFlags::AllowDecayPipelines;
 	ERHIPipeline PipelinesToBegin;
 	ERHIPipeline PipelinesToEnd;
 	bool bTransitionNeeded = false;
diff --git a/Engine/Source/Runtime/RenderCore/Public/RenderGraphTextureSubresource.h b/Engine/Source/Runtime/RenderCore/Public/RenderGraphTextureSubresource.h
--- a/Engine/Source/Runtime/RenderCore/Public/RenderGraphTextureSubresource.h
+++ b/Engine/Source/Runtime/RenderCore/Public/RenderGraphTextureSubresource.h
@@ -230,8 +230,7 @@ template <typename ElementType, typename AllocatorType>
 inline void InitTextureSubresources(TRDGTextureSubresourceArray<ElementType, AllocatorType>& SubresourceArray, const FRDGTextureSubresourceLayout& Layout, const ElementType& Element = {})
 {
 	const uint32 SubresourceCount = Layout.GetSubresourceCount();
-	SubresourceArray.Reserve(SubresourceCount);
-	SubresourceArray.SetNum(SubresourceCount);
+	SubresourceArray.SetNumUninitialized(SubresourceCount, EAllowShrinking::No);
 	for (uint32 SubresourceIndex = 0; SubresourceIndex < SubresourceCount; ++SubresourceIndex)
 	{
 		SubresourceArray[SubresourceIndex] = Element;
diff --git a/Engine/Source/Runtime/Renderer/Private/RayTracing/RayTracingDynamicGeometryUpdateManager.cpp b/Engine/Source/Runtime/Renderer/Private/RayTracing/RayTracingDynamicGeometryUpdateManager.cpp
--- a/Engine/Source/Runtime/Renderer/Private/RayTracing/RayTracingDynamicGeometryUpdateManager.cpp
+++ b/Engine/Source/Runtime/Renderer/Private/RayTracing/RayTracingDynamicGeometryUpdateManager.cpp
@@ -926,7 +926,7 @@ void FRayTracingDynamicGeometryUpdateManager::DispatchUpdates(FRHICommandList& R
 
 			// Transition to writeable for each cmd list and enable UAV overlap, because several dispatches can update non-overlapping portions of the same buffer.
 			// Mark as no fence because these resources have been fenced already at the beginning of the frame by RDG
-			RHICmdList.Transition(TransitionsBefore, ERHITransitionCreateFlags::NoFence);
+			RHICmdList.Transition(TransitionsBefore, ERHITransitionCreateFlags::AllowDecayPipelines);
 			RHICmdList.BeginUAVOverlap(OverlapUAVs);
 
 			// Cache the bound uniform buffers because a lot are the same between dispatches
diff --git a/Engine/Source/Runtime/SlateRHIRenderer/Private/SlateRHIRenderer.cpp b/Engine/Source/Runtime/SlateRHIRenderer/Private/SlateRHIRenderer.cpp
--- a/Engine/Source/Runtime/SlateRHIRenderer/Private/SlateRHIRenderer.cpp
+++ b/Engine/Source/Runtime/SlateRHIRenderer/Private/SlateRHIRenderer.cpp
@@ -968,7 +968,7 @@ void FSlateRHIRenderer::PresentWindow_RenderThread(FRHICommandListImmediate& RHI
 
 	RHICmdList.BeginDrawingViewport(DrawPassOutputs.ViewportRHI, FTextureRHIRef());
 
-	if (GRHIGlobals.NeedExtraTransitions)
+	if (GRHIGlobals.NeedsExtraTransitions)
 	{
 		RHICmdList.Transition(FRHITransitionInfo(DrawPassOutputs.OutputTextureRHI, ERHIAccess::Unknown, ERHIAccess::Present));
 
@@ -982,7 +982,7 @@ void FSlateRHIRenderer::PresentWindow_RenderThread(FRHICommandListImmediate& RHI
 
 	RHICmdList.EndDrawingViewport(DrawPassOutputs.ViewportRHI, true, DrawPassInputs.bLockToVsync);
 	
-	if (GRHIGlobals.NeedExtraTransitions)
+	if (GRHIGlobals.NeedsExtraTransitions)
 	{
 		RHICmdList.Transition(FRHITransitionInfo(DrawPassOutputs.OutputTextureRHI, ERHIAccess::Present, ERHIAccess::SRVMask));
 
diff --git a/Engine/Source/Runtime/VulkanRHI/Private/VulkanBarriers.cpp b/Engine/Source/Runtime/VulkanRHI/Private/VulkanBarriers.cpp
--- a/Engine/Source/Runtime/VulkanRHI/Private/VulkanBarriers.cpp
+++ b/Engine/Source/Runtime/VulkanRHI/Private/VulkanBarriers.cpp
@@ -5,8 +5,7 @@
 #include "VulkanContext.h"
 #include "VulkanPendingState.h"
 #include "ProfilingDebugging/RealtimeGPUProfiler.h"
-
-
+#include "RHICoreTransitions.h"
 
 // All shader stages supported by VK device - VK_PIPELINE_STAGE_VERTEX_SHADER_BIT, FRAGMENT etc
 extern uint32 GVulkanDevicePipelineStageBits;
@@ -647,6 +646,13 @@ static void SetupSubresourceRange(VkImageSubresourceRange& SubresRange, const FR
 	}
 }
 
+struct FVulkanTransitionData
+{
+	ERHIPipeline SrcPipelines, DstPipelines;
+	VulkanRHI::FSemaphore* Semaphore = nullptr;
+	TArray<FRHITransitionInfo> TransitionInfos;
+};
+
 void FVulkanDynamicRHI::RHICreateTransition(FRHITransition* Transition, const FRHITransitionCreateInfo& CreateInfo)
 {
 	TRACE_CPUPROFILER_EVENT_SCOPE(RHICreateTransition);
@@ -654,19 +660,41 @@ void FVulkanDynamicRHI::RHICreateTransition(FRHITransition* Transition, const FR
 	const ERHIPipeline SrcPipelines = CreateInfo.SrcPipelines;
 	const ERHIPipeline DstPipelines = CreateInfo.DstPipelines;
 
-	FVulkanPipelineBarrier* Data = new (Transition->GetPrivateData<FVulkanPipelineBarrier>()) FVulkanPipelineBarrier;
+	FVulkanTransitionData* Data = new (Transition->GetPrivateData<FVulkanTransitionData>()) FVulkanTransitionData;
+	Data->TransitionInfos = CreateInfo.TransitionInfos;
 	Data->SrcPipelines = SrcPipelines;
 	Data->DstPipelines = DstPipelines;
 
-	uint32 SrcQueueFamilyIndex, DstQueueFamilyIndex;
-	if (SrcPipelines != DstPipelines)
+	if (SrcPipelines != DstPipelines && !EnumHasAnyFlags(CreateInfo.Flags, ERHITransitionCreateFlags::NoFence))
 	{
 		Data->Semaphore = new VulkanRHI::FSemaphore(*Device);
+	}
+}
+
+void FVulkanDynamicRHI::RHIReleaseTransition(FRHITransition* Transition)
+{
+	Transition->GetPrivateData<FVulkanTransitionData>()->~FVulkanTransitionData();
+}
+
+// We need to keep the texture pointers around, because we need to call OnTransitionResource on them, and we need mip and layer counts for the tracking code.
+struct FImageBarrierExtraData
+{
+	FVulkanTexture* BaseTexture = nullptr;
+	bool IsAliasingBarrier = false;
+	uint8 PlaneSlice = 0;
+	ERHIAccess PlaneAccess = ERHIAccess::Unknown;
+};
 
-		const uint32 GfxQueueIndex = Device->GetGraphicsQueue()->GetFamilyIndex();
-		const uint32 ComputeQueueIndex = Device->GetComputeQueue()->GetFamilyIndex();
+template <typename ProcessorType>
+void ProcessTransitions(ProcessorType& Processor, FVulkanCommandListContext& Context, const FVulkanTransitionData& Data)
+{
+	uint32 SrcQueueFamilyIndex, DstQueueFamilyIndex;
+	if (Data.SrcPipelines != Data.DstPipelines)
+	{
+		const uint32 GfxQueueIndex = Context.Device.GetGraphicsQueue()->GetFamilyIndex();
+		const uint32 ComputeQueueIndex = Context.Device.GetComputeQueue()->GetFamilyIndex();
 
-		if (SrcPipelines == ERHIPipeline::Graphics)
+		if (Data.SrcPipelines == ERHIPipeline::Graphics)
 		{
 			SrcQueueFamilyIndex = GfxQueueIndex;
 			DstQueueFamilyIndex = ComputeQueueIndex;
@@ -685,14 +713,14 @@ void FVulkanDynamicRHI::RHICreateTransition(FRHITransition* Transition, const FR
 
 	// Count the images and buffers to be able to pre-allocate the arrays.
 	int32 NumTextures = 0, NumBuffers = 0;
-	for (const FRHITransitionInfo& Info : CreateInfo.TransitionInfos)
+	for (const FRHITransitionInfo& Info : Data.TransitionInfos)
 	{
 		if (!Info.Resource)
 		{
 			continue;
 		}
 
-		if (Info.AccessAfter == ERHIAccess::Discard)
+		if (EnumHasAnyFlags(Info.AccessAfter, ERHIAccess::Discard))
 		{
 			// Discard as a destination is a no-op
 			continue;
@@ -719,19 +747,19 @@ void FVulkanDynamicRHI::RHICreateTransition(FRHITransition* Transition, const FR
 			}
 		}
 
-		++NumBuffers;
+		if (Data.SrcPipelines != Data.DstPipelines)
+		{
+			++NumBuffers;
+		}
 	}
 
-	Data->ImageBarriers.Reserve(NumTextures);
-	Data->ImageBarrierExtras.Reserve(NumTextures);
-	if (SrcPipelines != DstPipelines)
-	{
-		Data->BufferBarriers.Reserve(NumBuffers);
-	}
+	Processor.ReserveBarriers(NumBuffers, NumTextures);
 
 	const ERHIAccess DepthStencilFlags = ERHIAccess::DSVRead | ERHIAccess::DSVWrite;
 
-	for (const FRHITransitionInfo& Info : CreateInfo.TransitionInfos)
+	const bool bNeedsAllPlanes = Context.Device.NeedsAllPlanes();
+
+	for (const FRHITransitionInfo& Info : Data.TransitionInfos)
 	{
 		if (!Info.Resource)
 		{
@@ -744,7 +772,7 @@ void FVulkanDynamicRHI::RHICreateTransition(FRHITransition* Transition, const FR
 			continue;
 		}
 
-		checkf(Info.AccessAfter != ERHIAccess::Unknown, TEXT("Transitioning a resource to an unknown state is not allowed."));
+		const UE::RHICore::FResourceState ResourceState(Context, Data.SrcPipelines, Data.DstPipelines, Info);
 
 		FVulkanBuffer* Buffer = nullptr;
 		FVulkanTexture* Texture = nullptr;
@@ -807,18 +835,18 @@ void FVulkanDynamicRHI::RHICreateTransition(FRHITransition* Transition, const FR
 		if (bIsDepthStencil)
 		{
 			// if we use separate transitions, then just feed them in as they are
-			SrcLayout = FVulkanPipelineBarrier::GetDepthOrStencilLayout(Info.AccessBefore);
-			DstLayout = FVulkanPipelineBarrier::GetDepthOrStencilLayout(Info.AccessAfter);
-			GetDepthOrStencilStageAndAccessFlags(Info.AccessBefore, SrcStageMask, SrcAccessFlags);
-			GetDepthOrStencilStageAndAccessFlags(Info.AccessAfter, DstStageMask, DstAccessFlags);
+			SrcLayout = FVulkanPipelineBarrier::GetDepthOrStencilLayout(ResourceState.AccessBefore);
+			DstLayout = FVulkanPipelineBarrier::GetDepthOrStencilLayout(ResourceState.AccessAfter);
+			GetDepthOrStencilStageAndAccessFlags(ResourceState.AccessBefore, SrcStageMask, SrcAccessFlags);
+			GetDepthOrStencilStageAndAccessFlags(ResourceState.AccessAfter, DstStageMask, DstAccessFlags);
 		}
 		else
 		{
 			
 			const bool bSupportsReadOnlyOptimal = (Texture == nullptr) || Texture->SupportsSampling();
 
-			GetVkStageAndAccessFlags(Info.AccessBefore, UnderlyingType, UsageFlags, bIsDepthStencil, bSupportsReadOnlyOptimal, SrcStageMask, SrcAccessFlags, SrcLayout, true);
-			GetVkStageAndAccessFlags(Info.AccessAfter, UnderlyingType, UsageFlags, bIsDepthStencil, bSupportsReadOnlyOptimal, DstStageMask, DstAccessFlags, DstLayout, false);
+			GetVkStageAndAccessFlags(ResourceState.AccessBefore, UnderlyingType, UsageFlags, bIsDepthStencil, bSupportsReadOnlyOptimal, SrcStageMask, SrcAccessFlags, SrcLayout, true);
+			GetVkStageAndAccessFlags(ResourceState.AccessAfter, UnderlyingType, UsageFlags, bIsDepthStencil, bSupportsReadOnlyOptimal, DstStageMask, DstAccessFlags, DstLayout, false);
 
 			// If not compute, remove vertex pipeline bits as only compute updates vertex buffers
 			if (!(SrcStageMask & VK_PIPELINE_STAGE_COMPUTE_SHADER_BIT))
@@ -828,17 +856,30 @@ void FVulkanDynamicRHI::RHICreateTransition(FRHITransition* Transition, const FR
 		}
 
 		// If we're not transitioning across pipes and we don't need to perform layout transitions, we can express memory dependencies through a global memory barrier.
-		if ((SrcPipelines == DstPipelines) && (Texture == nullptr || (SrcLayout == DstLayout)) && GVulkanUseMemoryBarrierOpt)
+		if ((Data.SrcPipelines == Data.DstPipelines) && (Texture == nullptr || (SrcLayout == DstLayout)) && GVulkanUseMemoryBarrierOpt)
 		{
-			Data->AddMemoryBarrier(SrcAccessFlags, DstAccessFlags, SrcStageMask, DstStageMask);
+			const VkAccessFlags ReadMask = VK_ACCESS_INDIRECT_COMMAND_READ_BIT | VK_ACCESS_INDEX_READ_BIT | VK_ACCESS_VERTEX_ATTRIBUTE_READ_BIT | VK_ACCESS_UNIFORM_READ_BIT |
+				VK_ACCESS_INPUT_ATTACHMENT_READ_BIT | VK_ACCESS_SHADER_READ_BIT | VK_ACCESS_COLOR_ATTACHMENT_READ_BIT | VK_ACCESS_DEPTH_STENCIL_ATTACHMENT_READ_BIT |
+				VK_ACCESS_TRANSFER_READ_BIT | VK_ACCESS_MEMORY_READ_BIT;
+
+			VkMemoryBarrier2 NewBarrier;
+			NewBarrier.sType = VK_STRUCTURE_TYPE_MEMORY_BARRIER_2;
+			NewBarrier.pNext = nullptr;
+			NewBarrier.srcStageMask = SrcStageMask;
+			NewBarrier.srcAccessMask = SrcAccessFlags;
+			NewBarrier.dstStageMask = DstStageMask;
+			NewBarrier.dstAccessMask = DstAccessFlags;
+
+			Processor.AddMemoryBarrier(NewBarrier);
 			continue;
 		}
 
 		if (Buffer != nullptr)
 		{
 			// We only add buffer transitions for cross-pipe transfers.
-			checkSlow( (SrcPipelines != DstPipelines) && (Texture == nullptr) );
-			VkBufferMemoryBarrier2& BufferBarrier = Data->BufferBarriers.AddDefaulted_GetRef();
+			checkSlow( (Data.SrcPipelines != Data.DstPipelines) && (Texture == nullptr) );
+
+			VkBufferMemoryBarrier2 BufferBarrier;
 			BufferBarrier.sType = VK_STRUCTURE_TYPE_BUFFER_MEMORY_BARRIER_2;
 			BufferBarrier.pNext = nullptr;
 			BufferBarrier.srcStageMask = SrcStageMask;
@@ -850,6 +891,8 @@ void FVulkanDynamicRHI::RHICreateTransition(FRHITransition* Transition, const FR
 			BufferBarrier.buffer = Buffer->GetHandle();
 			BufferBarrier.offset = 0;
 			BufferBarrier.size = VK_WHOLE_SIZE;
+
+			Processor.AddBufferBarrier(BufferBarrier);
 			continue;
 		}
 
@@ -862,29 +905,39 @@ void FVulkanDynamicRHI::RHICreateTransition(FRHITransition* Transition, const FR
 		// For some textures, e.g. FVulkanBackBuffer, the image handle may not be set yet, or may be stale, so there's no point storing it here.
 		// We'll set the image to NULL in the barrier info, and RHIEndTransitions will fetch the up to date pointer from the texture, after
 		// OnLayoutTransition is called.
-		VkImageMemoryBarrier2& ImgBarrier = Data->ImageBarriers.AddDefaulted_GetRef();
+		VkImageMemoryBarrier2 ImgBarrier;
 		SetupImageBarrier(ImgBarrier, VK_NULL_HANDLE, SrcStageMask, DstStageMask, SrcAccessFlags, DstAccessFlags, SrcLayout, DstLayout, SubresRange);
 		ImgBarrier.srcQueueFamilyIndex = SrcQueueFamilyIndex;
 		ImgBarrier.dstQueueFamilyIndex = DstQueueFamilyIndex;
 
-		FVulkanPipelineBarrier::ImageBarrierExtraData& ExtraData = Data->ImageBarrierExtras.AddDefaulted_GetRef();
+		FImageBarrierExtraData ExtraData;
 		ExtraData.BaseTexture = Texture;
-		ExtraData.IsAliasingBarrier = (Info.AccessBefore == ERHIAccess::Discard);
-		if (EnumHasAnyFlags(Info.Flags, EResourceTransitionFlags::FillOtherPlaneState))
-		{
-			checkSlow(bIsDepthStencil);
-			ExtraData.OtherPlaneLayout = FVulkanPipelineBarrier::GetDepthOrStencilLayout(Info.OtherPlaneAccess);
-		}
-		else
+		ExtraData.IsAliasingBarrier = (ResourceState.AccessBefore == ERHIAccess::Discard);
+
+		if (bNeedsAllPlanes && Texture->IsDepthOrStencilAspect())
 		{
-			checkSlow(!bIsDepthStencil || (Info.PlaneSlice == FRHISubresourceRange::kAllSubresources) || AspectFlags == Texture->GetFullAspectMask());
+			ExtraData.PlaneAccess = ResourceState.AccessAfter;
+			ExtraData.PlaneSlice  = Info.PlaneSlice;
 		}
+
+		Processor.AddImageBarrier(ImgBarrier, ExtraData);
 	}
 }
 
-void FVulkanDynamicRHI::RHIReleaseTransition(FRHITransition* Transition)
+template <typename MemoryBarrierType>
+static void MergeBarrierAccessMask(MemoryBarrierType& MemoryBarrier, VkAccessFlags InSrcAccessFlags, VkAccessFlags InDstAccessFlags)
 {
-	Transition->GetPrivateData<FVulkanPipelineBarrier>()->~FVulkanPipelineBarrier();
+	const VkAccessFlags ReadMask = VK_ACCESS_INDIRECT_COMMAND_READ_BIT | VK_ACCESS_INDEX_READ_BIT | VK_ACCESS_VERTEX_ATTRIBUTE_READ_BIT | VK_ACCESS_UNIFORM_READ_BIT |
+		VK_ACCESS_INPUT_ATTACHMENT_READ_BIT | VK_ACCESS_SHADER_READ_BIT | VK_ACCESS_COLOR_ATTACHMENT_READ_BIT | VK_ACCESS_DEPTH_STENCIL_ATTACHMENT_READ_BIT |
+		VK_ACCESS_TRANSFER_READ_BIT | VK_ACCESS_MEMORY_READ_BIT;
+
+	// We only need a memory barrier if the previous commands wrote to the buffer. In case of a transition from read, an execution barrier is enough.
+	const bool SrcAccessIsRead = ((InSrcAccessFlags & (~ReadMask)) == 0);
+	if (!SrcAccessIsRead)
+	{
+		MemoryBarrier.srcAccessMask |= InSrcAccessFlags;
+		MemoryBarrier.dstAccessMask |= InDstAccessFlags;
+	}
 }
 
 static void DowngradeBarrier(VkMemoryBarrier& OutBarrier, const VkMemoryBarrier2& InBarrier)
@@ -922,6 +975,15 @@ static void DowngradeBarrier(VkImageMemoryBarrier& OutBarrier, const VkImageMemo
 	OutBarrier.subresourceRange = InBarrier.subresourceRange;
 }
 
+template <typename DstArrayType, typename BarrierType>
+static void DowngradeBarrier(DstArrayType& TargetArray, const BarrierType& SrcBarrier, VkPipelineStageFlags& MergedSrcStageMask, VkPipelineStageFlags& MergedDstStageMask)
+{
+	auto& DstBarrier = TargetArray.AddDefaulted_GetRef();
+	DowngradeBarrier(DstBarrier, SrcBarrier);
+	MergedSrcStageMask |= SrcBarrier.srcStageMask;
+	MergedDstStageMask |= SrcBarrier.dstStageMask;
+}
+
 template <typename DstArrayType, typename SrcArrayType>
 static void DowngradeBarrierArray(DstArrayType& TargetArray, const SrcArrayType& SrcArray, VkPipelineStageFlags& MergedSrcStageMask, VkPipelineStageFlags& MergedDstStageMask)
 {
@@ -977,9 +1039,11 @@ public:
 
 	void Process(TArrayView<const FRHITransition*>& Transitions)
 	{
+		const bool bNeedsAllPlanes = Context.Device.NeedsAllPlanes();
+
 		for (const FRHITransition* Transition : Transitions)
 		{
-			const FVulkanPipelineBarrier* Data = Transition->GetPrivateData<FVulkanPipelineBarrier>();
+			const FVulkanTransitionData* Data = Transition->GetPrivateData<FVulkanTransitionData>();
 
 			const bool IsCrossPipe = (Data->SrcPipelines != Data->DstPipelines);
 
@@ -993,7 +1057,8 @@ public:
 
 			const int32 BufferBatchStartIndex = BufferBarriers.Num();
 			const int32 ImageBatchStartIndex = ImageBarriers.Num();
-			AppendBarriers(Data);
+
+			ProcessTransitions(*this, Context, *Data);
 
 			// Patch the src/dst access mask for acquire/release on queue ownership transfers
 			if (IsCrossPipe)
@@ -1003,12 +1068,17 @@ public:
 			}
 
 			// Set the images for the barriers from the texture in the extra data and notify of layout changes
-			for (int32 Index = 0; Index < Data->ImageBarrierExtras.Num(); ++Index)
+			for (int32 Index = 0; Index < ImageBarrierExtras.Num(); ++Index)
 			{
-				const FVulkanPipelineBarrier::ImageBarrierExtraData* RemainingExtras = &Data->ImageBarrierExtras[Index];
+				const FImageBarrierExtraData* RemainingExtras = &ImageBarrierExtras[Index];
 				FVulkanTexture* Texture = RemainingExtras->BaseTexture;
 				check(Texture->Image != VK_NULL_HANDLE || !Texture->IsImageOwner());  // coming in, we should always have an image, except backbuffer with r.Vulkan.DelayAcquireBackBuffer=2
 
+				if (bNeedsAllPlanes && Texture->IsDepthOrStencilAspect())
+				{
+					Texture->AllPlanesTrackedAccess[RemainingExtras->PlaneSlice] = RemainingExtras->PlaneAccess;
+				}
+
 				const int32 TargetIndex = ImageBatchStartIndex + Index;
 				ImageBarrierType* RemainingBarriers = &ImageBarriers[TargetIndex];
 
@@ -1022,7 +1092,7 @@ public:
 
 				if ((RemainingBarriers->image != VK_NULL_HANDLE) && (RemainingBarriers->subresourceRange.aspectMask != 0))
 				{
-					const int32 RemainingCount = Data->ImageBarrierExtras.Num() - Index;
+					const int32 RemainingCount = ImageBarrierExtras.Num() - Index;
 					ApplyTracking(TargetIndex, RemainingExtras, RemainingCount, IsCrossPipe);
 				}
 			}
@@ -1043,15 +1113,27 @@ public:
 				}
 			}
 
+			ImageBarrierExtras.Reset();
+
 			FinishBatch(Data->SrcPipelines, Data->DstPipelines, BufferBatchStartIndex, ImageBatchStartIndex);
 		}
 
 		FinishAll();
 	}
 
+	void ReserveBarriers(int32 NumBuffers, int32 NumTextures)
+	{
+		BufferBarriers.Reserve(BufferBarriers.Num() + NumBuffers);
+		ImageBarriers.Reserve(ImageBarriers.Num() + NumTextures);
+		ImageBarrierExtras.Reserve(NumTextures);
+	}
+
+	void AddMemoryBarrier(const VkMemoryBarrier2& Barrier);
+	void AddBufferBarrier(const VkBufferMemoryBarrier2& Barrier);
+	void AddImageBarrier(const VkImageMemoryBarrier2& Barrier, const FImageBarrierExtraData& ExtraData);
+
 private:
-	void AppendBarriers(const FVulkanPipelineBarrier* Data);
-	void ApplyTracking(uint32 TargetIndex, const FVulkanPipelineBarrier::ImageBarrierExtraData* RemainingExtras, int32 RemainingCount, bool IsCrossPipe);
+	void ApplyTracking(uint32 TargetIndex, const FImageBarrierExtraData* RemainingExtras, int32 RemainingCount, bool IsCrossPipe);
 	void FinishBatch(ERHIPipeline SrcPipelines, ERHIPipeline DstPipelines, int32 BufferBatchStartIndex, int32 ImageBatchStartIndex);
 	void FinishAll();
 
@@ -1078,7 +1160,7 @@ private:
 		}
 	}
 
-	void ValidateQueue(const FVulkanPipelineBarrier* Data)
+	void ValidateQueue(const FVulkanTransitionData* Data)
 	{
 #if DO_GUARD_SLOW
 		ERHIPipeline RequestedPipeline = IsBeginTransition ? Data->SrcPipelines : Data->DstPipelines;
@@ -1100,6 +1182,7 @@ private:
 	TArray<MemoryBarrierType> MemoryBarriers;
 	TArray<BufferBarrierType> BufferBarriers;
 	TArray<ImageBarrierType> ImageBarriers;
+	TArray<FImageBarrierExtraData> ImageBarrierExtras;
 };
 
 
@@ -1138,16 +1221,35 @@ void FTransitionProcessor<VkMemoryBarrier, VkBufferMemoryBarrier, VkImageMemoryB
 	// Non-Sync2 already submitted each batch
 }
 
+
 template <>
-void FTransitionProcessor<VkMemoryBarrier, VkBufferMemoryBarrier, VkImageMemoryBarrier>::AppendBarriers(const FVulkanPipelineBarrier* Data)
+void FTransitionProcessor<VkMemoryBarrier, VkBufferMemoryBarrier, VkImageMemoryBarrier>::AddMemoryBarrier(const VkMemoryBarrier2& Barrier)
 {
-	DowngradeBarrierArray(MemoryBarriers, Data->MemoryBarriers, MergedSrcStageMask, MergedDstStageMask);
-	DowngradeBarrierArray(BufferBarriers, Data->BufferBarriers, MergedSrcStageMask, MergedDstStageMask);
-	DowngradeBarrierArray(ImageBarriers, Data->ImageBarriers, MergedSrcStageMask, MergedDstStageMask);
+	if (MemoryBarriers.Num() == 0)
+	{
+		VkMemoryBarrier& NewBarrier = MemoryBarriers.AddDefaulted_GetRef();
+		ZeroVulkanStruct(NewBarrier, VK_STRUCTURE_TYPE_MEMORY_BARRIER);
+	}
+
+	// Mash everything into a single barrier
+	MergeBarrierAccessMask(MemoryBarriers[0], Barrier.srcAccessMask, Barrier.dstAccessMask);
 }
 
 template <>
-void FTransitionProcessor<VkMemoryBarrier, VkBufferMemoryBarrier, VkImageMemoryBarrier>::ApplyTracking(uint32_t TargetIndex, const FVulkanPipelineBarrier::ImageBarrierExtraData* RemainingExtras, int32 RemainingCount, bool IsCrossPipe)
+void FTransitionProcessor<VkMemoryBarrier, VkBufferMemoryBarrier, VkImageMemoryBarrier>::AddBufferBarrier(const VkBufferMemoryBarrier2& Barrier)
+{
+	DowngradeBarrier(BufferBarriers, Barrier, MergedSrcStageMask, MergedDstStageMask);
+}
+
+template <>
+void FTransitionProcessor<VkMemoryBarrier, VkBufferMemoryBarrier, VkImageMemoryBarrier>::AddImageBarrier(const VkImageMemoryBarrier2& Barrier, const FImageBarrierExtraData& ExtraData)
+{
+	DowngradeBarrier(ImageBarriers, Barrier, MergedSrcStageMask, MergedDstStageMask);
+	ImageBarrierExtras.Emplace(ExtraData);
+}
+
+template <>
+void FTransitionProcessor<VkMemoryBarrier, VkBufferMemoryBarrier, VkImageMemoryBarrier>::ApplyTracking(uint32_t TargetIndex, const FImageBarrierExtraData* RemainingExtras, int32 RemainingCount, bool IsCrossPipe)
 {
 	VkImageMemoryBarrier& ImageBarrier = ImageBarriers[TargetIndex];
 	FVulkanTexture* Texture = RemainingExtras->BaseTexture;
@@ -1190,9 +1292,9 @@ void FTransitionProcessor<VkMemoryBarrier, VkBufferMemoryBarrier, VkImageMemoryB
 			}
 			else
 			{
-				checkf((RemainingExtras->OtherPlaneLayout != VK_IMAGE_LAYOUT_MAX_ENUM), TEXT("This shouldn't happen with GRHIGlobals.NeedsAllPlanes set!"));
+				ERHIAccess OtherPlaneAccess = Texture->AllPlanesTrackedAccess[ImageBarrier.subresourceRange.aspectMask == VK_IMAGE_ASPECT_DEPTH_BIT ? 0 : 1];
 
-				OtherAspectOldLayout = OtherAspectNewLayout = RemainingExtras->OtherPlaneLayout;
+				OtherAspectOldLayout = OtherAspectNewLayout = FVulkanPipelineBarrier::GetDepthOrStencilLayout(OtherPlaneAccess);
 			}
 
 			// Merge the layout with its other half and set it in the barrier
@@ -1322,15 +1424,37 @@ void FTransitionProcessor<VkMemoryBarrier2, VkBufferMemoryBarrier2, VkImageMemor
 }
 
 template <>
-void FTransitionProcessor<VkMemoryBarrier2, VkBufferMemoryBarrier2, VkImageMemoryBarrier2>::AppendBarriers(const FVulkanPipelineBarrier* Data)
+void FTransitionProcessor<VkMemoryBarrier2, VkBufferMemoryBarrier2, VkImageMemoryBarrier2>::AddMemoryBarrier(const VkMemoryBarrier2& Barrier)
+{
+	if (MemoryBarriers.Num() == 0)
+	{
+		VkMemoryBarrier2& NewBarrier = MemoryBarriers.AddDefaulted_GetRef();
+		ZeroVulkanStruct(NewBarrier, VK_STRUCTURE_TYPE_MEMORY_BARRIER_2);
+	}
+
+	// Mash everything into a single barrier
+	VkMemoryBarrier2& MemoryBarrier = MemoryBarriers[0];
+
+	MergeBarrierAccessMask(MemoryBarrier, Barrier.srcAccessMask, Barrier.dstAccessMask);
+	MemoryBarrier.srcStageMask |= Barrier.srcStageMask;
+	MemoryBarrier.dstStageMask |= Barrier.dstStageMask;
+}
+
+template <>
+void FTransitionProcessor<VkMemoryBarrier2, VkBufferMemoryBarrier2, VkImageMemoryBarrier2>::AddBufferBarrier(const VkBufferMemoryBarrier2& Barrier)
 {
-	MemoryBarriers.Append(Data->MemoryBarriers);
-	BufferBarriers.Append(Data->BufferBarriers);
-	ImageBarriers.Append(Data->ImageBarriers);
+	BufferBarriers.Emplace(Barrier);
 }
 
 template <>
-void FTransitionProcessor<VkMemoryBarrier2, VkBufferMemoryBarrier2, VkImageMemoryBarrier2>::ApplyTracking(uint32 TargetIndex, const FVulkanPipelineBarrier::ImageBarrierExtraData* RemainingExtras, int32 RemainingCount, bool IsCrossPipe)
+void FTransitionProcessor<VkMemoryBarrier2, VkBufferMemoryBarrier2, VkImageMemoryBarrier2>::AddImageBarrier(const VkImageMemoryBarrier2& Barrier, const FImageBarrierExtraData& ExtraData)
+{
+	ImageBarriers.Emplace(Barrier);
+	ImageBarrierExtras.Emplace(ExtraData);
+}
+
+template <>
+void FTransitionProcessor<VkMemoryBarrier2, VkBufferMemoryBarrier2, VkImageMemoryBarrier2>::ApplyTracking(uint32 TargetIndex, const FImageBarrierExtraData* RemainingExtras, int32 RemainingCount, bool IsCrossPipe)
 {
 	// Nothing to do when VK_KHR_separate_depth_stencil_layouts and VK_KHR_synchronization2 are supported (requirements for SM6).
 }
@@ -1345,7 +1469,7 @@ static TArray<VulkanRHI::FSemaphore*> ExtractTransitionSemaphores(TArrayView<con
 
 	for (const FRHITransition* Transition : Transitions)
 	{
-		const FVulkanPipelineBarrier* Data = Transition->GetPrivateData<FVulkanPipelineBarrier>();
+		const FVulkanTransitionData* Data = Transition->GetPrivateData<FVulkanTransitionData>();
 		if (Data->Semaphore == nullptr)
 		{
 			continue;
@@ -1405,13 +1529,8 @@ void FVulkanCommandListContext::RHIEndTransitions(TArrayView<const FRHITransitio
 	}
 }
 
-
 void FVulkanPipelineBarrier::AddMemoryBarrier(VkAccessFlags InSrcAccessFlags, VkAccessFlags InDstAccessFlags, VkPipelineStageFlags InSrcStageMask, VkPipelineStageFlags InDstStageMask)
 {
-	const VkAccessFlags ReadMask = VK_ACCESS_INDIRECT_COMMAND_READ_BIT | VK_ACCESS_INDEX_READ_BIT | VK_ACCESS_VERTEX_ATTRIBUTE_READ_BIT | VK_ACCESS_UNIFORM_READ_BIT |
-		VK_ACCESS_INPUT_ATTACHMENT_READ_BIT | VK_ACCESS_SHADER_READ_BIT | VK_ACCESS_COLOR_ATTACHMENT_READ_BIT | VK_ACCESS_DEPTH_STENCIL_ATTACHMENT_READ_BIT |
-		VK_ACCESS_TRANSFER_READ_BIT | VK_ACCESS_MEMORY_READ_BIT;
-
 	if (MemoryBarriers.Num() == 0)
 	{
 		VkMemoryBarrier2& NewBarrier = MemoryBarriers.AddDefaulted_GetRef();
@@ -1421,19 +1540,11 @@ void FVulkanPipelineBarrier::AddMemoryBarrier(VkAccessFlags InSrcAccessFlags, Vk
 	// Mash everything into a single barrier
 	VkMemoryBarrier2& MemoryBarrier = MemoryBarriers[0];
 
-	// We only need a memory barrier if the previous commands wrote to the buffer. In case of a transition from read, an execution barrier is enough.
-	const bool SrcAccessIsRead = ((InSrcAccessFlags & (~ReadMask)) == 0);
-	if (!SrcAccessIsRead)
-	{
-		MemoryBarrier.srcAccessMask |= InSrcAccessFlags;
-		MemoryBarrier.dstAccessMask |= InDstAccessFlags;
-	}
-
+	MergeBarrierAccessMask(MemoryBarrier, InSrcAccessFlags, InDstAccessFlags);
 	MemoryBarrier.srcStageMask |= InSrcStageMask;
 	MemoryBarrier.dstStageMask |= InDstStageMask;
 }
 
-
 //
 // Methods used when the RHI itself needs to perform a layout transition. The public API functions do not call these,
 // they fill in the fields of FVulkanPipelineBarrier using their own logic, based on the ERHIAccess flags.
diff --git a/Engine/Source/Runtime/VulkanRHI/Private/VulkanBarriers.h b/Engine/Source/Runtime/VulkanRHI/Private/VulkanBarriers.h
--- a/Engine/Source/Runtime/VulkanRHI/Private/VulkanBarriers.h
+++ b/Engine/Source/Runtime/VulkanRHI/Private/VulkanBarriers.h
@@ -12,27 +12,15 @@ class FVulkanCommandBuffer;
 
 struct FVulkanPipelineBarrier
 {
-	FVulkanPipelineBarrier() : Semaphore(nullptr)
-	{}
+	FVulkanPipelineBarrier() = default;
 
 	using MemoryBarrierArrayType = TArray<VkMemoryBarrier2, TInlineAllocator<1>>;
 	using ImageBarrierArrayType = TArray<VkImageMemoryBarrier2, TInlineAllocator<2>>;
 	using BufferBarrierArrayType = TArray<VkBufferMemoryBarrier2>;
 
-	ERHIPipeline SrcPipelines, DstPipelines;
 	MemoryBarrierArrayType MemoryBarriers;
 	ImageBarrierArrayType ImageBarriers;
 	BufferBarrierArrayType BufferBarriers;
-	VulkanRHI::FSemaphore* Semaphore;
-
-	// We need to keep the texture pointers around, because we need to call OnTransitionResource on them, and we need mip and layer counts for the tracking code.
-	struct ImageBarrierExtraData
-	{
-		FVulkanTexture* BaseTexture = nullptr;
-		VkImageLayout OtherPlaneLayout = VK_IMAGE_LAYOUT_MAX_ENUM;
-		bool IsAliasingBarrier = false;
-	};
-	TArray<ImageBarrierExtraData, TInlineAllocator<2>> ImageBarrierExtras;
 
 	void AddMemoryBarrier(VkAccessFlags SrcAccessFlags, VkAccessFlags DstAccessFlags, VkPipelineStageFlags SrcStageMask, VkPipelineStageFlags DstStageMask);
 	void AddFullImageLayoutTransition(const FVulkanTexture& Texture, VkImageLayout SrcLayout, VkImageLayout DstLayout);
@@ -46,4 +34,4 @@ struct FVulkanPipelineBarrier
 	// Predetermined layouts for given RHIAccess
 	static VkImageLayout GetDefaultLayout(const FVulkanTexture& VulkanTexture, ERHIAccess DesiredAccess);
 	static VkImageLayout GetDepthOrStencilLayout(ERHIAccess Access);
-};
+};
\ No newline at end of file
diff --git a/Engine/Source/Runtime/VulkanRHI/Private/VulkanDevice.h b/Engine/Source/Runtime/VulkanRHI/Private/VulkanDevice.h
--- a/Engine/Source/Runtime/VulkanRHI/Private/VulkanDevice.h
+++ b/Engine/Source/Runtime/VulkanRHI/Private/VulkanDevice.h
@@ -521,7 +521,12 @@ public:
 		return OptionalDeviceExtensionProperties;
 	}
 
-	bool SupportsParallelRendering() const
+	inline bool NeedsAllPlanes() const
+	{
+		return !OptionalDeviceExtensions.HasSeparateDepthStencilLayouts;
+	}
+
+	inline bool SupportsParallelRendering() const
 	{
 		return OptionalDeviceExtensions.HasSeparateDepthStencilLayouts && OptionalDeviceExtensions.HasKHRSynchronization2 && OptionalDeviceExtensions.HasKHRRenderPass2;
 	}
diff --git a/Engine/Source/Runtime/VulkanRHI/Private/VulkanRHI.cpp b/Engine/Source/Runtime/VulkanRHI/Private/VulkanRHI.cpp
--- a/Engine/Source/Runtime/VulkanRHI/Private/VulkanRHI.cpp
+++ b/Engine/Source/Runtime/VulkanRHI/Private/VulkanRHI.cpp
@@ -856,9 +856,7 @@ void FVulkanDynamicRHI::InitInstance()
 		GRHISupportsRayTracing = RHI_RAYTRACING && RHISupportsRayTracing(GMaxRHIShaderPlatform) && Device->GetOptionalExtensions().HasRaytracingExtensions();
 		GRHIGlobals.SupportsMapWriteNoOverwrite = true;
 
-		GRHIGlobals.EnableHighLevelTransitionPatching = true;
-		GRHIGlobals.NeedExtraTransitions = true;
-		GRHIGlobals.NeedsAllPlanes = true;
+		GRHIGlobals.NeedsExtraTransitions = true;
 
 		// Use this compatibility mode avoid known issues at launch time with latest drivers at the time of release 5.5.  This will:
 		// - disable inline ray tracing and use ray tracing pipelines everywhere (instead of a mix of both)
diff --git a/Engine/Source/Runtime/VulkanRHI/Private/VulkanRenderTarget.cpp b/Engine/Source/Runtime/VulkanRHI/Private/VulkanRenderTarget.cpp
--- a/Engine/Source/Runtime/VulkanRHI/Private/VulkanRenderTarget.cpp
+++ b/Engine/Source/Runtime/VulkanRHI/Private/VulkanRenderTarget.cpp
@@ -518,6 +518,8 @@ void FVulkanCommandListContext::RHIBeginRenderPass(const FRHIRenderPassInfo& InI
 		BeginOcclusionQueryBatch(InInfo.NumOcclusionQueries);
 	}
 
+	const bool bNeedsAllPlanes = Device.NeedsAllPlanes();
+
 	FRHITexture* DSTexture = InInfo.DepthStencilRenderTarget.DepthStencilTarget;
 	VkImageLayout CurrentDepthLayout = VK_IMAGE_LAYOUT_UNDEFINED;
 	VkImageLayout CurrentStencilLayout = VK_IMAGE_LAYOUT_UNDEFINED;
@@ -537,9 +539,9 @@ void FVulkanCommandListContext::RHIBeginRenderPass(const FRHIRenderPassInfo& InI
 			{
 				CurrentDepthLayout = VK_IMAGE_LAYOUT_READ_ONLY_OPTIMAL;
 			}
-			else
+			else if (bNeedsAllPlanes)
 			{
-				CurrentDepthLayout = FVulkanPipelineBarrier::GetDepthOrStencilLayout(InInfo.NopAccess);
+				CurrentDepthLayout = FVulkanPipelineBarrier::GetDepthOrStencilLayout(VulkanTexture.AllPlanesTrackedAccess[0]);
 			}
 		}
 
@@ -553,9 +555,9 @@ void FVulkanCommandListContext::RHIBeginRenderPass(const FRHIRenderPassInfo& InI
 			{
 				CurrentStencilLayout = VK_IMAGE_LAYOUT_READ_ONLY_OPTIMAL;
 			}
-			else
+			else if (bNeedsAllPlanes)
 			{
-				CurrentStencilLayout = FVulkanPipelineBarrier::GetDepthOrStencilLayout(InInfo.NopAccess);
+				CurrentStencilLayout = FVulkanPipelineBarrier::GetDepthOrStencilLayout(VulkanTexture.AllPlanesTrackedAccess[1]);
 			}
 		}
 	}
diff --git a/Engine/Source/Runtime/VulkanRHI/Private/VulkanResources.h b/Engine/Source/Runtime/VulkanRHI/Private/VulkanResources.h
--- a/Engine/Source/Runtime/VulkanRHI/Private/VulkanResources.h
+++ b/Engine/Source/Runtime/VulkanRHI/Private/VulkanResources.h
@@ -622,7 +622,7 @@ public:
 
 	// FVulkanEvictable interface.
 	bool CanMove() const override { return false; }
-	bool CanEvict() const override { return AreAllSubresourcesSame_Unsafe(); }
+	bool CanEvict() const override { return false; }
 	void Evict(FVulkanDevice& Device, FVulkanCommandListContext& Context) override; ///evict to system memory
 	void Move(FVulkanDevice& Device, FVulkanCommandListContext& Context, VulkanRHI::FVulkanAllocation& NewAllocation) override; //move to a full new allocation
 	FVulkanTexture* GetEvictableTexture() override { return CanEvict() ? this : nullptr; }
@@ -775,6 +775,9 @@ public:
 
 	FVulkanRHIExternalImageDeleteCallbackInfo ExternalImageDeleteCallbackInfo;
 
+	// Only used when HasSeparateDepthStencilLayouts == false.
+	TStaticArray<ERHIAccess, 2> AllPlanesTrackedAccess{ InPlace, ERHIAccess::Unknown };
+
 	void SetInitialImageState(FRHICommandListBase& RHICmdList, VkImageLayout InitialLayout, bool bClear, const FClearValueBinding& ClearValueBinding, bool bIsTransientResource);
 
 private:
diff --git a/Engine/Source/Runtime/VulkanRHI/Private/VulkanTexture.cpp b/Engine/Source/Runtime/VulkanRHI/Private/VulkanTexture.cpp
--- a/Engine/Source/Runtime/VulkanRHI/Private/VulkanTexture.cpp
+++ b/Engine/Source/Runtime/VulkanRHI/Private/VulkanTexture.cpp
@@ -1009,6 +1009,14 @@ FVulkanDynamicRHI::FCreateTextureResult FVulkanDynamicRHI::BeginCreateTextureInt
 
 	FVulkanTexture* Texture = new FVulkanTexture(*Device, CreateDesc, InTransientHeapAllocation);
 
+	const bool bNeedsAllPlanes = Device->NeedsAllPlanes();
+
+	if (bNeedsAllPlanes)
+	{
+		Texture->AllPlanesTrackedAccess[0] = CreateDesc.InitialState;
+		Texture->AllPlanesTrackedAccess[1] = CreateDesc.InitialState;
+	}
+
 	const bool bIsTransientResource = (InTransientHeapAllocation != nullptr);
 	const bool bDoInitialClear = VKHasAnyFlags(Texture->ImageUsageFlags, VK_IMAGE_USAGE_SAMPLED_BIT) &&
 		EnumHasAnyFlags(CreateDesc.Flags, TexCreate_RenderTargetable | TexCreate_DepthStencilTargetable) &&
