VulkanRHI: Add VK_KHR_maintenance7 to loaded extensions.

#rnx

diff --git a/Engine/Source/Runtime/VulkanRHI/Private/VulkanDevice.h b/Engine/Source/Runtime/VulkanRHI/Private/VulkanDevice.h
--- a/Engine/Source/Runtime/VulkanRHI/Private/VulkanDevice.h
+++ b/Engine/Source/Runtime/VulkanRHI/Private/VulkanDevice.h
@@ -63,6 +63,7 @@ struct FOptionalVulkanDeviceExtensions
 			uint64 HasEXTToolingInfo : 1;
 			uint64 HasEXTImageCompressionControl : 1;
 			uint64 HasEXTMutableDescriptorType : 1;
+			uint64 HasKHRMaintenance7 : 1;
 
 			// Vendor specific
 			uint64 HasAMDBufferMarker : 1;
@@ -135,6 +136,7 @@ struct FOptionalVulkanDeviceExtensionProperties
 	}
 
 	VkPhysicalDeviceDriverPropertiesKHR PhysicalDeviceDriverProperties;
+	VkPhysicalDeviceMaintenance7PropertiesKHR PhysicalDeviceMaintenance7Properties;
 
 	VkPhysicalDeviceDescriptorBufferPropertiesEXT DescriptorBufferProps;
 	VkPhysicalDeviceSubgroupSizeControlPropertiesEXT SubgroupSizeControlProperties;
diff --git a/Engine/Source/Runtime/VulkanRHI/Private/VulkanExtensions.cpp b/Engine/Source/Runtime/VulkanRHI/Private/VulkanExtensions.cpp
--- a/Engine/Source/Runtime/VulkanRHI/Private/VulkanExtensions.cpp
+++ b/Engine/Source/Runtime/VulkanRHI/Private/VulkanExtensions.cpp
@@ -232,6 +232,54 @@ private:
 };
 
 
+// ***** VK_KHR_maintenance7
+class FVulkanKHRMaintenance7Extension : public FVulkanDeviceExtension
+{
+public:
+
+	FVulkanKHRMaintenance7Extension(FVulkanDevice* InDevice)
+		: FVulkanDeviceExtension(InDevice, VK_KHR_MAINTENANCE_7_EXTENSION_NAME, VULKAN_EXTENSION_ENABLED)
+	{
+	}
+
+	virtual void PrePhysicalDeviceFeatures(VkPhysicalDeviceFeatures2KHR& PhysicalDeviceFeatures2) override final
+	{
+		ZeroVulkanStruct(Maintenance7Features, VK_STRUCTURE_TYPE_PHYSICAL_DEVICE_MAINTENANCE_7_FEATURES_KHR);
+		AddToPNext(PhysicalDeviceFeatures2, Maintenance7Features);
+	}
+
+	virtual void PostPhysicalDeviceFeatures(FOptionalVulkanDeviceExtensions& ExtensionFlags) override final
+	{
+		bRequirementsPassed = (Maintenance7Features.maintenance7 == VK_TRUE);
+		if (bRequirementsPassed)
+		{
+			ExtensionFlags.HasKHRMaintenance7 = true;
+		}
+	}
+
+	virtual void PrePhysicalDeviceProperties(VkPhysicalDeviceProperties2KHR& PhysicalDeviceProperties2) override final
+	{
+		if (bRequirementsPassed)
+		{
+			VkPhysicalDeviceMaintenance7PropertiesKHR& PhysicalDeviceMaintenance7Properties = GetDeviceExtensionProperties().PhysicalDeviceMaintenance7Properties;
+			ZeroVulkanStruct(PhysicalDeviceMaintenance7Properties, VK_STRUCTURE_TYPE_PHYSICAL_DEVICE_MAINTENANCE_7_PROPERTIES_KHR);
+			AddToPNext(PhysicalDeviceProperties2, PhysicalDeviceMaintenance7Properties);
+		}
+	}
+
+	virtual void PreCreateDevice(VkDeviceCreateInfo& DeviceCreateInfo) override final
+	{
+		if (bRequirementsPassed)
+		{
+			AddToPNext(DeviceCreateInfo, Maintenance7Features);
+		}
+	}
+
+private:
+	VkPhysicalDeviceMaintenance7FeaturesKHR Maintenance7Features;
+};
+
+
 // ***** VK_KHR_driver_properties
 class FVulkanKHRDriverPropertiesExtension : public FVulkanDeviceExtension
 {
@@ -1826,6 +1874,7 @@ FVulkanDeviceExtensionArray FVulkanDeviceExtension::GetUESupportedDeviceExtensio
 	ADD_CUSTOM_EXTENSION(FVulkanKHRDriverPropertiesExtension);
 	ADD_CUSTOM_EXTENSION(FVulkanKHRMaintenance4Extension);
 	ADD_CUSTOM_EXTENSION(FVulkanKHRMaintenance5Extension);
+	ADD_CUSTOM_EXTENSION(FVulkanKHRMaintenance7Extension);
 	ADD_CUSTOM_EXTENSION(FVulkanShaderAtomicInt64Extension);
 	ADD_CUSTOM_EXTENSION(FVulkanShaderImageAtomicInt64Extension);
 	ADD_CUSTOM_EXTENSION(FVulkanEXTScalarBlockLayoutExtension);
