VulkanRHI: Fix assert in Debug about transition data size/alignment.

[RN]
#rnx

diff --git a/Engine/Source/Runtime/VulkanRHI/Private/VulkanBarriers.cpp b/Engine/Source/Runtime/VulkanRHI/Private/VulkanBarriers.cpp
--- a/Engine/Source/Runtime/VulkanRHI/Private/VulkanBarriers.cpp
+++ b/Engine/Source/Runtime/VulkanRHI/Private/VulkanBarriers.cpp
@@ -646,13 +646,6 @@ static void SetupSubresourceRange(VkImageSubresourceRange& SubresRange, const FR
 	}
 }
 
-struct FVulkanTransitionData
-{
-	ERHIPipeline SrcPipelines, DstPipelines;
-	VulkanRHI::FSemaphore* Semaphore = nullptr;
-	TArray<FRHITransitionInfo> TransitionInfos;
-};
-
 void FVulkanDynamicRHI::RHICreateTransition(FRHITransition* Transition, const FRHITransitionCreateInfo& CreateInfo)
 {
 	TRACE_CPUPROFILER_EVENT_SCOPE(RHICreateTransition);
diff --git a/Engine/Source/Runtime/VulkanRHI/Private/VulkanBarriers.h b/Engine/Source/Runtime/VulkanRHI/Private/VulkanBarriers.h
--- a/Engine/Source/Runtime/VulkanRHI/Private/VulkanBarriers.h
+++ b/Engine/Source/Runtime/VulkanRHI/Private/VulkanBarriers.h
@@ -34,4 +34,12 @@ struct FVulkanPipelineBarrier
 	// Predetermined layouts for given RHIAccess
 	static VkImageLayout GetDefaultLayout(const FVulkanTexture& VulkanTexture, ERHIAccess DesiredAccess);
 	static VkImageLayout GetDepthOrStencilLayout(ERHIAccess Access);
-};
\ No newline at end of file
+};
+
+struct FVulkanTransitionData
+{
+	ERHIPipeline SrcPipelines, DstPipelines;
+	VulkanRHI::FSemaphore* Semaphore = nullptr;  // used for cross queue
+	VkEvent EventHandle = VK_NULL_HANDLE;  // used for split barriers
+	TArray<FRHITransitionInfo> TransitionInfos;
+};
diff --git a/Engine/Source/Runtime/VulkanRHI/Private/VulkanRHI.cpp b/Engine/Source/Runtime/VulkanRHI/Private/VulkanRHI.cpp
--- a/Engine/Source/Runtime/VulkanRHI/Private/VulkanRHI.cpp
+++ b/Engine/Source/Runtime/VulkanRHI/Private/VulkanRHI.cpp
@@ -440,8 +440,8 @@ FVulkanDynamicRHI::FVulkanDynamicRHI()
 	GTexturePoolSize = 0;
 	GRHISupportsMultithreading = true;
 	GRHISupportsMultithreadedResources = true;
-	GRHITransitionPrivateData_SizeInBytes = sizeof(FVulkanPipelineBarrier);
-	GRHITransitionPrivateData_AlignInBytes = alignof(FVulkanPipelineBarrier);
+	GRHITransitionPrivateData_SizeInBytes = sizeof(FVulkanTransitionData);
+	GRHITransitionPrivateData_AlignInBytes = alignof(FVulkanTransitionData);
 	GConfig->GetInt(TEXT("TextureStreaming"), TEXT("PoolSizeVRAMPercentage"), GPoolSizeVRAMPercentage, GEngineIni);
 
 	GRHIGlobals.SupportsBarycentricsSemantic = true;
