VulkanRHI: Fix validation issue in previews by adding UB flag to transient memory heaps.

#rnx

diff --git a/Engine/Source/Runtime/VulkanRHI/Private/VulkanTransientResourceAllocator.cpp b/Engine/Source/Runtime/VulkanRHI/Private/VulkanTransientResourceAllocator.cpp
--- a/Engine/Source/Runtime/VulkanRHI/Private/VulkanTransientResourceAllocator.cpp
+++ b/Engine/Source/Runtime/VulkanRHI/Private/VulkanTransientResourceAllocator.cpp
@@ -12,7 +12,7 @@ FVulkanTransientHeap::FVulkanTransientHeap(const FInitializer& Initializer, FVul
 	, FDeviceChild(InDevice)
 	, VulkanBuffer(VK_NULL_HANDLE)
 {
-	EBufferUsageFlags UEBufferUsageFlags = BUF_VertexBuffer | BUF_IndexBuffer | BUF_DrawIndirect 
+	EBufferUsageFlags UEBufferUsageFlags = BUF_UniformBuffer | BUF_VertexBuffer | BUF_IndexBuffer | BUF_DrawIndirect 
 		| BUF_UnorderedAccess | BUF_StructuredBuffer | BUF_ShaderResource | BUF_KeepCPUAccessible;
 
 	if (InDevice->GetOptionalExtensions().HasRaytracingExtensions())
