VulkanRHI: Fix memory barriers for legacy barriers.  Match NeedsAllPlanes usage to !SupportsParallelRendering.  Remove LazyAcquire option to remove need for OnLayoutTransition, was not supported on Android anyhow.  Make sync2 a direct requirement for parallel rendering.  Fix barriers for async compute.  Add support for split barriers on desktop.  Remove FTransitionProcessor and split the code paths for legacy and sync2 barriers to allow for more optimal code paths on either side.

#rnx

diff --git a/Engine/Config/Android/AndroidEngine.ini b/Engine/Config/Android/AndroidEngine.ini
--- a/Engine/Config/Android/AndroidEngine.ini
+++ b/Engine/Config/Android/AndroidEngine.ini
@@ -91,6 +91,8 @@ r.Vulkan.Submission.UseSubmissionThread=0
 r.Vulkan.RHIThread=1
 ; Disable async compute until we can target devices to enable it on
 r.Vulkan.AllowAsyncCompute=0
+; Disable split barriers until its effects are measured on all devices
+r.Vulkan.AllowSplitBarriers=0
 
 [/Script/Engine.RendererSettings]
 r.DistanceFields=1
diff --git a/Engine/Source/Runtime/VulkanRHI/Private/VulkanBarriers.cpp b/Engine/Source/Runtime/VulkanRHI/Private/VulkanBarriers.cpp
--- a/Engine/Source/Runtime/VulkanRHI/Private/VulkanBarriers.cpp
+++ b/Engine/Source/Runtime/VulkanRHI/Private/VulkanBarriers.cpp
@@ -10,6 +10,8 @@
 // All shader stages supported by VK device - VK_PIPELINE_STAGE_VERTEX_SHADER_BIT, FRAGMENT etc
 extern uint32 GVulkanDevicePipelineStageBits;
 
+extern int32 GVulkanAllowConcurrentBuffer;
+extern int32 GVulkanAllowConcurrentImage;
 
 int32 GVulkanUseMemoryBarrierOpt = 1;
 static FAutoConsoleVariableRef CVarVulkanUseMemoryBarrierOpt(
@@ -31,6 +33,16 @@ static FAutoConsoleVariableRef CVarVulkanMaxBarriersPerBatch(
 	ECVF_Default
 );
 
+int32 GVulkanAllowSplitBarriers = 1;
+static FAutoConsoleVariableRef CVarVulkanAllowSplitBarriers(
+	TEXT("r.Vulkan.AllowSplitBarriers"),
+	GVulkanAllowSplitBarriers,
+	TEXT("Will limit the number of barriers sent per batch\n")
+	TEXT(" 0: Disable split barriers\n")
+	TEXT(" 1: Allow split barriers using Synchronization2 events (default)\n"),
+	ECVF_Default
+);
+
 
 //
 // The following two functions are used when the RHI needs to do image layout transitions internally.
@@ -573,19 +585,6 @@ static void GetDepthOrStencilStageAndAccessFlags(ERHIAccess Access, VkPipelineSt
 }
 
 
-// Helper To apply to each bit in an aspect
-template <typename TFunction>
-static void ForEachAspect(VkImageAspectFlags AspectFlags, TFunction Function)
-{
-	// Keep it simple for now, can iterate on more bits when needed
-	for (uint32 SingleBit = 1; SingleBit <= VK_IMAGE_ASPECT_STENCIL_BIT; SingleBit <<= 1)
-	{
-		if ((AspectFlags & SingleBit) != 0)
-		{
-			Function((VkImageAspectFlagBits)SingleBit);
-		}
-	}
-}
 
 //
 // Helpers for filling in the fields of a VkImageMemoryBarrier structure.
@@ -658,7 +657,7 @@ void FVulkanDynamicRHI::RHICreateTransition(FRHITransition* Transition, const FR
 	Data->SrcPipelines = SrcPipelines;
 	Data->DstPipelines = DstPipelines;
 
-	if (SrcPipelines != DstPipelines && !EnumHasAnyFlags(CreateInfo.Flags, ERHITransitionCreateFlags::NoFence))
+	if ((SrcPipelines != DstPipelines) && !EnumHasAnyFlags(CreateInfo.Flags, ERHITransitionCreateFlags::NoFence))
 	{
 		Data->Semaphore = new VulkanRHI::FSemaphore(*Device);
 	}
@@ -678,31 +677,35 @@ struct FImageBarrierExtraData
 	ERHIAccess PlaneAccess = ERHIAccess::Unknown;
 };
 
-template <typename ProcessorType>
-void ProcessTransitions(ProcessorType& Processor, FVulkanCommandListContext& Context, const FVulkanTransitionData& Data)
+
+struct FLegacyBarrierArrays
 {
-	uint32 SrcQueueFamilyIndex, DstQueueFamilyIndex;
-	if (Data.SrcPipelines != Data.DstPipelines)
-	{
-		const uint32 GfxQueueIndex = Context.Device.GetGraphicsQueue()->GetFamilyIndex();
-		const uint32 ComputeQueueIndex = Context.Device.GetComputeQueue()->GetFamilyIndex();
+	TArray<VkMemoryBarrier, TInlineAllocator<1>> MemoryBarriers;
+	TArray<VkBufferMemoryBarrier> BufferBarriers;
+	TArray<VkImageMemoryBarrier> ImageBarriers;
+	TArray<FImageBarrierExtraData> ImageExtraData;
 
-		if (Data.SrcPipelines == ERHIPipeline::Graphics)
-		{
-			SrcQueueFamilyIndex = GfxQueueIndex;
-			DstQueueFamilyIndex = ComputeQueueIndex;
-		}
-		else
-		{
-			SrcQueueFamilyIndex = ComputeQueueIndex;
-			DstQueueFamilyIndex = GfxQueueIndex;
-		}
-	}
-	else
-	{
-		SrcQueueFamilyIndex = VK_QUEUE_FAMILY_IGNORED;
-		DstQueueFamilyIndex = VK_QUEUE_FAMILY_IGNORED;
-	}
+	VkPipelineStageFlags SrcStageMask = 0;
+	VkPipelineStageFlags DstStageMask = 0;
+};
+
+struct FSync2BarrierArrays
+{
+	TArray<VkMemoryBarrier2, TInlineAllocator<1>> MemoryBarriers;
+	TArray<VkBufferMemoryBarrier2> BufferBarriers;
+	TArray<VkImageMemoryBarrier2> ImageBarriers;
+};
+
+
+template<typename BarrierArrayType>
+void ConvertTransitionToBarriers(FVulkanCommandListContext& Context, const FVulkanTransitionData& Data, BarrierArrayType& OutBarriers)
+{
+	using MemoryBarrierType = typename decltype(OutBarriers.MemoryBarriers)::ElementType;
+	using BufferBarrierType = typename decltype(OutBarriers.BufferBarriers)::ElementType;
+	using ImageBarrierType = typename decltype(OutBarriers.ImageBarriers)::ElementType;
+
+	// Legacy barriers use a single stage for all barriers
+	constexpr bool bIsLegacyBarrier = std::is_same_v<FLegacyBarrierArrays, BarrierArrayType>;
 
 	// Count the images and buffers to be able to pre-allocate the arrays.
 	int32 NumTextures = 0, NumBuffers = 0;
@@ -746,11 +749,16 @@ void ProcessTransitions(ProcessorType& Processor, FVulkanCommandListContext& Con
 		}
 	}
 
-	Processor.ReserveBarriers(NumBuffers, NumTextures);
 
-	const ERHIAccess DepthStencilFlags = ERHIAccess::DSVRead | ERHIAccess::DSVWrite;
+	// Presize all the arrays
+	if (!GVulkanUseMemoryBarrierOpt)
+	{
+		OutBarriers.BufferBarriers.Reserve(OutBarriers.BufferBarriers.Num() + NumBuffers);
+	}
+	OutBarriers.ImageBarriers.Reserve(OutBarriers.ImageBarriers.Num() + NumTextures);
+
 
-	const bool bNeedsAllPlanes = Context.Device.NeedsAllPlanes();
+	const ERHIAccess DepthStencilFlags = ERHIAccess::DSVRead | ERHIAccess::DSVWrite;
 
 	for (const FRHITransitionInfo& Info : Data.TransitionInfos)
 	{
@@ -848,72 +856,101 @@ void ProcessTransitions(ProcessorType& Processor, FVulkanCommandListContext& Con
 			}
 		}
 
+		// Mash them all together for legacy barriers, set them for each barrier in sync2
+		if constexpr (bIsLegacyBarrier)
+		{
+			OutBarriers.SrcStageMask |= SrcStageMask;
+			OutBarriers.DstStageMask |= DstStageMask;
+		}
+
 		// If we're not transitioning across pipes and we don't need to perform layout transitions, we can express memory dependencies through a global memory barrier.
 		if ((Data.SrcPipelines == Data.DstPipelines) && (Texture == nullptr || (SrcLayout == DstLayout)) && GVulkanUseMemoryBarrierOpt)
 		{
+			if (OutBarriers.MemoryBarriers.Num() == 0)
+			{
+				MemoryBarrierType& NewBarrier = OutBarriers.MemoryBarriers.AddZeroed_GetRef();
+				NewBarrier.sType = bIsLegacyBarrier ? VK_STRUCTURE_TYPE_MEMORY_BARRIER : VK_STRUCTURE_TYPE_MEMORY_BARRIER_2;
+			}
+
 			const VkAccessFlags ReadMask = VK_ACCESS_INDIRECT_COMMAND_READ_BIT | VK_ACCESS_INDEX_READ_BIT | VK_ACCESS_VERTEX_ATTRIBUTE_READ_BIT | VK_ACCESS_UNIFORM_READ_BIT |
 				VK_ACCESS_INPUT_ATTACHMENT_READ_BIT | VK_ACCESS_SHADER_READ_BIT | VK_ACCESS_COLOR_ATTACHMENT_READ_BIT | VK_ACCESS_DEPTH_STENCIL_ATTACHMENT_READ_BIT |
 				VK_ACCESS_TRANSFER_READ_BIT | VK_ACCESS_MEMORY_READ_BIT;
 
-			VkMemoryBarrier2 NewBarrier;
-			NewBarrier.sType = VK_STRUCTURE_TYPE_MEMORY_BARRIER_2;
-			NewBarrier.pNext = nullptr;
-			NewBarrier.srcStageMask = SrcStageMask;
-			NewBarrier.srcAccessMask = SrcAccessFlags;
-			NewBarrier.dstStageMask = DstStageMask;
-			NewBarrier.dstAccessMask = DstAccessFlags;
+			// Mash everything into a single barrier
+			MemoryBarrierType& MemoryBarrier = OutBarriers.MemoryBarriers[0];
 
-			Processor.AddMemoryBarrier(NewBarrier);
-			continue;
-		}
+			// We only need a memory barrier if the previous commands wrote to the buffer. In case of a transition from read, an execution barrier is enough.
+			const bool SrcAccessIsRead = ((SrcAccessFlags & (~ReadMask)) == 0);
+			if (!SrcAccessIsRead)
+			{
+				MemoryBarrier.srcAccessMask |= SrcAccessFlags;
+				MemoryBarrier.dstAccessMask |= DstAccessFlags;
+			}
 
-		if (Buffer != nullptr)
+			if constexpr (!bIsLegacyBarrier)
+			{
+				MemoryBarrier.srcStageMask = SrcStageMask;
+				MemoryBarrier.dstStageMask = DstStageMask;
+			}
+		}
+		else if (Buffer != nullptr)
 		{
-			// We only add buffer transitions for cross-pipe transfers.
-			checkSlow( (Data.SrcPipelines != Data.DstPipelines) && (Texture == nullptr) );
-
-			VkBufferMemoryBarrier2 BufferBarrier;
-			BufferBarrier.sType = VK_STRUCTURE_TYPE_BUFFER_MEMORY_BARRIER_2;
-			BufferBarrier.pNext = nullptr;
-			BufferBarrier.srcStageMask = SrcStageMask;
-			BufferBarrier.dstStageMask = DstStageMask;
+			BufferBarrierType& BufferBarrier = OutBarriers.BufferBarriers.AddZeroed_GetRef();
+			BufferBarrier.sType = bIsLegacyBarrier ? VK_STRUCTURE_TYPE_BUFFER_MEMORY_BARRIER : VK_STRUCTURE_TYPE_BUFFER_MEMORY_BARRIER_2;
 			BufferBarrier.srcAccessMask = SrcAccessFlags;
 			BufferBarrier.dstAccessMask = DstAccessFlags;
-			BufferBarrier.srcQueueFamilyIndex = SrcQueueFamilyIndex;
-			BufferBarrier.dstQueueFamilyIndex = DstQueueFamilyIndex;
-			BufferBarrier.buffer = Buffer->GetHandle();
-			BufferBarrier.offset = 0;
-			BufferBarrier.size = VK_WHOLE_SIZE;
+			BufferBarrier.srcQueueFamilyIndex = VK_QUEUE_FAMILY_IGNORED;
+			BufferBarrier.dstQueueFamilyIndex = VK_QUEUE_FAMILY_IGNORED;
 
-			Processor.AddBufferBarrier(BufferBarrier);
-			continue;
-		}
+			const VulkanRHI::FVulkanAllocation& BufferAlloc = Buffer->GetCurrentAllocation();
+			BufferBarrier.buffer = BufferAlloc.GetBufferHandle();
+			BufferBarrier.offset = BufferAlloc.Offset;
+			BufferBarrier.size = BufferAlloc.Size;
 
-		check(Texture != nullptr);
+			if constexpr (!bIsLegacyBarrier)
+			{
+				BufferBarrier.srcStageMask = SrcStageMask;
+				BufferBarrier.dstStageMask = DstStageMask;
+			}
+		}
+		else if (Texture != nullptr)
+		{
+			const VkImageAspectFlags AspectFlags = bIsDepthStencil ? GetDepthStencilAspectMask(Info.PlaneSlice) : Texture->GetFullAspectMask();
 
-		const VkImageAspectFlags AspectFlags = bIsDepthStencil ? GetDepthStencilAspectMask(Info.PlaneSlice) : Texture->GetFullAspectMask();
-		VkImageSubresourceRange SubresRange;
-		SetupSubresourceRange(SubresRange, Info, AspectFlags);
+			ImageBarrierType& ImageBarrier = OutBarriers.ImageBarriers.AddZeroed_GetRef();
+			ImageBarrier.sType = bIsLegacyBarrier ? VK_STRUCTURE_TYPE_IMAGE_MEMORY_BARRIER : VK_STRUCTURE_TYPE_IMAGE_MEMORY_BARRIER_2;
+			ImageBarrier.srcAccessMask = SrcAccessFlags;
+			ImageBarrier.dstAccessMask = DstAccessFlags;
+			ImageBarrier.oldLayout = SrcLayout;
+			ImageBarrier.newLayout = DstLayout;
+			ImageBarrier.srcQueueFamilyIndex = VK_QUEUE_FAMILY_IGNORED;
+			ImageBarrier.dstQueueFamilyIndex = VK_QUEUE_FAMILY_IGNORED;
+			ImageBarrier.image = Texture->Image;
 
-		// For some textures, e.g. FVulkanBackBuffer, the image handle may not be set yet, or may be stale, so there's no point storing it here.
-		// We'll set the image to NULL in the barrier info, and RHIEndTransitions will fetch the up to date pointer from the texture, after
-		// OnLayoutTransition is called.
-		VkImageMemoryBarrier2 ImgBarrier;
-		SetupImageBarrier(ImgBarrier, VK_NULL_HANDLE, SrcStageMask, DstStageMask, SrcAccessFlags, DstAccessFlags, SrcLayout, DstLayout, SubresRange);
-		ImgBarrier.srcQueueFamilyIndex = SrcQueueFamilyIndex;
-		ImgBarrier.dstQueueFamilyIndex = DstQueueFamilyIndex;
+			SetupSubresourceRange(ImageBarrier.subresourceRange, Info, AspectFlags);
 
-		FImageBarrierExtraData ExtraData;
-		ExtraData.BaseTexture = Texture;
-		ExtraData.IsAliasingBarrier = (ResourceState.AccessBefore == ERHIAccess::Discard);
+			if constexpr (bIsLegacyBarrier)
+			{
+				FImageBarrierExtraData& ExtraData = OutBarriers.ImageExtraData.AddDefaulted_GetRef();;
+				ExtraData.BaseTexture = Texture;
+				ExtraData.IsAliasingBarrier = (ResourceState.AccessBefore == ERHIAccess::Discard);
 
-		if (bNeedsAllPlanes && Texture->IsDepthOrStencilAspect())
+				if (Texture->IsDepthOrStencilAspect())
+				{
+					ExtraData.PlaneAccess = ResourceState.AccessAfter;
+					ExtraData.PlaneSlice = Info.PlaneSlice;
+				}
+			}
+			else
+			{
+				ImageBarrier.srcStageMask = SrcStageMask;
+				ImageBarrier.dstStageMask = DstStageMask;
+			}
+		}
+		else
 		{
-			ExtraData.PlaneAccess = ResourceState.AccessAfter;
-			ExtraData.PlaneSlice  = Info.PlaneSlice;
+			checkf(false, TEXT("Transition with no resource!"));
 		}
-
-		Processor.AddImageBarrier(ImgBarrier, ExtraData);
 	}
 }
 
@@ -1019,232 +1056,19 @@ static void MergeDepthStencilLayouts(BarrierArrayType& TargetArray)
 	}
 }
 
-// The transition processor takes RHITransitions and transforms them into Vulkan barriers (of the template type)
-template<typename MemoryBarrierType, typename BufferBarrierType, typename ImageBarrierType>
-class FTransitionProcessor
-{
-public:
-
-	FTransitionProcessor(FVulkanCommandListContext& InContext, bool InIsBeginTransition)
-		: Context(InContext)
-		, IsBeginTransition(InIsBeginTransition)
-	{}
-
-	void Process(TArrayView<const FRHITransition*>& Transitions)
-	{
-		const bool bNeedsAllPlanes = Context.Device.NeedsAllPlanes();
-
-		for (const FRHITransition* Transition : Transitions)
-		{
-			const FVulkanTransitionData* Data = Transition->GetPrivateData<FVulkanTransitionData>();
-
-			const bool IsCrossPipe = (Data->SrcPipelines != Data->DstPipelines);
-
-			// We only care about cross-pipe transitions when we begin transitions
-			if (IsBeginTransition && !IsCrossPipe)
-			{
-				continue;
-			}
-
-			ValidateQueue(Data);
-
-			const int32 BufferBatchStartIndex = BufferBarriers.Num();
-			const int32 ImageBatchStartIndex = ImageBarriers.Num();
-
-			ProcessTransitions(*this, Context, *Data);
-
-			// Patch the src/dst access mask for acquire/release on queue ownership transfers
-			if (IsCrossPipe)
-			{
-				PatchCrossPipeTransitions(BufferBarriers, BufferBatchStartIndex);
-				PatchCrossPipeTransitions(ImageBarriers, ImageBatchStartIndex);
-			}
-
-			// Set the images for the barriers from the texture in the extra data and notify of layout changes
-			for (int32 Index = 0; Index < ImageBarrierExtras.Num(); ++Index)
-			{
-				const FImageBarrierExtraData* RemainingExtras = &ImageBarrierExtras[Index];
-				FVulkanTexture* Texture = RemainingExtras->BaseTexture;
-				check(Texture->Image != VK_NULL_HANDLE || !Texture->IsImageOwner());  // coming in, we should always have an image, except backbuffer with r.Vulkan.DelayAcquireBackBuffer=2
 
-				if (bNeedsAllPlanes && Texture->IsDepthOrStencilAspect())
-				{
-					Texture->AllPlanesTrackedAccess[RemainingExtras->PlaneSlice] = RemainingExtras->PlaneAccess;
-				}
 
-				const int32 TargetIndex = ImageBatchStartIndex + Index;
-				ImageBarrierType* RemainingBarriers = &ImageBarriers[TargetIndex];
 
-				// EndTransition handles layout transition notifications
-				if (!IsBeginTransition)
-				{
-					Texture->OnLayoutTransition(Context, RemainingBarriers->newLayout);
-				}
 
-				RemainingBarriers->image = Texture->Image;
 
-				if ((RemainingBarriers->image != VK_NULL_HANDLE) && (RemainingBarriers->subresourceRange.aspectMask != 0))
-				{
-					const int32 RemainingCount = ImageBarrierExtras.Num() - Index;
-					ApplyTracking(TargetIndex, RemainingExtras, RemainingCount, IsCrossPipe);
-				}
-			}
 
-			// Loop again to remove invalidated entries
-			for (int32 DstIndex = ImageBatchStartIndex; DstIndex < ImageBarriers.Num();)
-			{
-				// If Texture is a backbuffer, OnLayoutTransition normally updates the image handle. However, if the viewport got invalidated during the
-				// OnLayoutTransition call, we may end up with a null handle, which is fine.  Merged Depth and Stencil transitions will also result in
-				// in null aspectMask for the extra transition.  Either is fine, just remove the transition.
-				if ((ImageBarriers[DstIndex].image == VK_NULL_HANDLE) || ((ImageBarriers[DstIndex].subresourceRange.aspectMask == 0)))
-				{
-					ImageBarriers.RemoveAtSwap(DstIndex);
-				}
-				else
-				{
-					++DstIndex;
-				}
-			}
 
-			ImageBarrierExtras.Reset();
 
-			FinishBatch(Data->SrcPipelines, Data->DstPipelines, BufferBatchStartIndex, ImageBatchStartIndex);
-		}
 
-		FinishAll();
-	}
-
-	void ReserveBarriers(int32 NumBuffers, int32 NumTextures)
-	{
-		BufferBarriers.Reserve(BufferBarriers.Num() + NumBuffers);
-		ImageBarriers.Reserve(ImageBarriers.Num() + NumTextures);
-		ImageBarrierExtras.Reserve(NumTextures);
-	}
-
-	void AddMemoryBarrier(const VkMemoryBarrier2& Barrier);
-	void AddBufferBarrier(const VkBufferMemoryBarrier2& Barrier);
-	void AddImageBarrier(const VkImageMemoryBarrier2& Barrier, const FImageBarrierExtraData& ExtraData);
-
-private:
-	void ApplyTracking(uint32 TargetIndex, const FImageBarrierExtraData* RemainingExtras, int32 RemainingCount, bool IsCrossPipe);
-	void FinishBatch(ERHIPipeline SrcPipelines, ERHIPipeline DstPipelines, int32 BufferBatchStartIndex, int32 ImageBatchStartIndex);
-	void FinishAll();
-
-	// Patches for acquire/release of resources during queue ownership transfers
-	// For async compute, sanitize the stage and access flags to what is supported by the queue
-	template<typename ArrayType>
-	void PatchCrossPipeTransitions(ArrayType& Barriers, int32 StartIndex)
-	{
-		for (int32 Index = StartIndex; Index < Barriers.Num(); ++Index)
-		{
-			auto& Barrier = Barriers[Index];
-
-			if (Barrier.srcQueueFamilyIndex != Barrier.dstQueueFamilyIndex)
-			{
-				if (IsBeginTransition)
-				{
-					Barrier.dstAccessMask = 0; // Release resource from current queue.
-				}
-				else
-				{
-					Barrier.srcAccessMask = 0; // Acquire resource on current queue.
-				}
-			}
-		}
-	}
-
-	void ValidateQueue(const FVulkanTransitionData* Data)
-	{
-#if DO_GUARD_SLOW
-		ERHIPipeline RequestedPipeline = IsBeginTransition ? Data->SrcPipelines : Data->DstPipelines;
-		checkf(EnumHasAnyFlags(Context.GetPipeline(), RequestedPipeline), 
-			TEXT("The %s pipelines for this transition are [%s], but it's submitted on the [%s] queue."), 
-			IsBeginTransition ? TEXT("SRC") : TEXT("DST"),
-			*GetRHIPipelineName(RequestedPipeline),
-			*GetRHIPipelineName(Context.GetPipeline())
-		);
-#endif
-	}
-
-	FVulkanCommandListContext& Context;
-	const bool IsBeginTransition;
-
-	VkPipelineStageFlags MergedSrcStageMask = 0;
-	VkPipelineStageFlags MergedDstStageMask = 0;
-
-	TArray<MemoryBarrierType> MemoryBarriers;
-	TArray<BufferBarrierType> BufferBarriers;
-	TArray<ImageBarrierType> ImageBarriers;
-	TArray<FImageBarrierExtraData> ImageBarrierExtras;
-};
-
-
-
-// *** FTransitionProcessor for Original Barriers ***
-
-template <> 
-void FTransitionProcessor<VkMemoryBarrier, VkBufferMemoryBarrier, VkImageMemoryBarrier>::FinishBatch(ERHIPipeline SrcPipelines, ERHIPipeline DstPipelines, int32 BufferBatchStartIndex, int32 ImageBatchStartIndex)
-{
-	if ((MemoryBarriers.Num() == 0) && (BufferBarriers.Num() == 0) && (ImageBarriers.Num() == 0))
-	{
-		return;
-	}
-
-	if (SrcPipelines == ERHIPipeline::AsyncCompute)
-	{
-		MergedSrcStageMask &= Context.Device.GetComputeQueue()->GetSupportedStageBits();
-	}
-	if (DstPipelines == ERHIPipeline::AsyncCompute)
-	{
-		MergedDstStageMask &= Context.Device.GetComputeQueue()->GetSupportedStageBits();
-	}
-
-	// Submit merged stage masks with arrays of barriers
-	VulkanRHI::vkCmdPipelineBarrier(Context.GetCommandBuffer().GetHandle(), MergedSrcStageMask, MergedDstStageMask, 0, MemoryBarriers.Num(), MemoryBarriers.GetData(), BufferBarriers.Num(), BufferBarriers.GetData(), ImageBarriers.Num(), ImageBarriers.GetData());
-	MemoryBarriers.Reset();
-	BufferBarriers.Reset();
-	ImageBarriers.Reset();
-	MergedSrcStageMask = 0;
-	MergedDstStageMask = 0;
-}
-
-template <> 
-void FTransitionProcessor<VkMemoryBarrier, VkBufferMemoryBarrier, VkImageMemoryBarrier>::FinishAll()
-{ 
-	// Non-Sync2 already submitted each batch
-}
-
-
-template <>
-void FTransitionProcessor<VkMemoryBarrier, VkBufferMemoryBarrier, VkImageMemoryBarrier>::AddMemoryBarrier(const VkMemoryBarrier2& Barrier)
-{
-	if (MemoryBarriers.Num() == 0)
-	{
-		VkMemoryBarrier& NewBarrier = MemoryBarriers.AddDefaulted_GetRef();
-		ZeroVulkanStruct(NewBarrier, VK_STRUCTURE_TYPE_MEMORY_BARRIER);
-	}
-
-	// Mash everything into a single barrier
-	MergeBarrierAccessMask(MemoryBarriers[0], Barrier.srcAccessMask, Barrier.dstAccessMask);
-}
-
-template <>
-void FTransitionProcessor<VkMemoryBarrier, VkBufferMemoryBarrier, VkImageMemoryBarrier>::AddBufferBarrier(const VkBufferMemoryBarrier2& Barrier)
-{
-	DowngradeBarrier(BufferBarriers, Barrier, MergedSrcStageMask, MergedDstStageMask);
-}
-
-template <>
-void FTransitionProcessor<VkMemoryBarrier, VkBufferMemoryBarrier, VkImageMemoryBarrier>::AddImageBarrier(const VkImageMemoryBarrier2& Barrier, const FImageBarrierExtraData& ExtraData)
+// Legacy barriers must always submit depth and stencil together
+static void MergePlanes(VkImageMemoryBarrier* ImageBarriers, const FImageBarrierExtraData* RemainingExtras, int32 RemainingCount)
 {
-	DowngradeBarrier(ImageBarriers, Barrier, MergedSrcStageMask, MergedDstStageMask);
-	ImageBarrierExtras.Emplace(ExtraData);
-}
-
-template <>
-void FTransitionProcessor<VkMemoryBarrier, VkBufferMemoryBarrier, VkImageMemoryBarrier>::ApplyTracking(uint32_t TargetIndex, const FImageBarrierExtraData* RemainingExtras, int32 RemainingCount, bool IsCrossPipe)
-{
-	VkImageMemoryBarrier& ImageBarrier = ImageBarriers[TargetIndex];
+	VkImageMemoryBarrier& ImageBarrier = ImageBarriers[0];
 	FVulkanTexture* Texture = RemainingExtras->BaseTexture;
 	check(Texture->Image == ImageBarrier.image);
 	check(ImageBarrier.newLayout != VK_IMAGE_LAYOUT_UNDEFINED);
@@ -1261,7 +1085,7 @@ void FTransitionProcessor<VkMemoryBarrier, VkBufferMemoryBarrier, VkImageMemoryB
 			VkImageMemoryBarrier* OtherAspectImageBarrier = nullptr;
 			for (int32 OtherBarrierIndex = 0; OtherBarrierIndex < RemainingCount; ++OtherBarrierIndex)
 			{
-				VkImageMemoryBarrier& OtherImageBarrier = ImageBarriers[TargetIndex + OtherBarrierIndex];
+				VkImageMemoryBarrier& OtherImageBarrier = ImageBarriers[OtherBarrierIndex];
 				FVulkanTexture* OtherTexture = RemainingExtras[OtherBarrierIndex].BaseTexture;
 				if ((OtherTexture->Image == ImageBarrier.image) && (OtherImageBarrier.subresourceRange.aspectMask == OtherAspectMask))
 				{
@@ -1312,60 +1136,197 @@ void FTransitionProcessor<VkMemoryBarrier, VkBufferMemoryBarrier, VkImageMemoryB
 		}
 	}
 
-	MergedSrcStageMask |= GetVkStageFlagsForLayout(ImageBarrier.oldLayout);
-	MergedDstStageMask |= GetVkStageFlagsForLayout(ImageBarrier.newLayout);
-
 	// Once we're done with the barrier, make sure there are no sync2 states left
 	check((ImageBarrier.oldLayout != VK_IMAGE_LAYOUT_READ_ONLY_OPTIMAL) && (ImageBarrier.oldLayout != VK_IMAGE_LAYOUT_ATTACHMENT_OPTIMAL));
 	check((ImageBarrier.newLayout != VK_IMAGE_LAYOUT_READ_ONLY_OPTIMAL) && (ImageBarrier.newLayout != VK_IMAGE_LAYOUT_ATTACHMENT_OPTIMAL));
 }
 
+// Create Vulkan barriers from RHI transitions when VK_KHR_Synchronization2 is NOT supported (legacy code path)
+void ProcessTransitionLegacy(FVulkanCommandListContext& Context, TArrayView<const FRHITransition*>& Transitions)
+{
+	for (const FRHITransition* Transition : Transitions)
+	{
+		const FVulkanTransitionData* Data = Transition->GetPrivateData<FVulkanTransitionData>();
+
+		const bool bIsSingleQueue = IsSingleRHIPipeline(Data->SrcPipelines) && (Data->SrcPipelines == Data->DstPipelines);
+		checkf(bIsSingleQueue, TEXT("Devices without support for Sync2 should not be using async compute."));
+
+#if DO_GUARD_SLOW
+		checkf(EnumHasAnyFlags(Data->SrcPipelines, Context.GetPipeline()) && EnumHasAnyFlags(Data->DstPipelines, Context.GetPipeline()),
+			TEXT("The pipelines for this transition are [%s -> %s], but it's submitted on the [%s] queue."),
+			*GetRHIPipelineName(Data->SrcPipelines),
+			*GetRHIPipelineName(Data->DstPipelines),
+			*GetRHIPipelineName(Context.GetPipeline())
+		);
+#endif
+
+		FLegacyBarrierArrays BarrierArrays;
+		ConvertTransitionToBarriers(Context, *Data, BarrierArrays);
 
+		// Merge any depth/stencil barriers
+		for (int32 Index = 0; Index < BarrierArrays.ImageExtraData.Num(); ++Index)
+		{
+			const FImageBarrierExtraData* RemainingExtras = &BarrierArrays.ImageExtraData[Index];
+			VkImageMemoryBarrier* RemainingBarriers = &BarrierArrays.ImageBarriers[Index];
 
-// *** FTransitionProcessor for Sync2 Barriers ***
+			if (RemainingExtras->BaseTexture->IsDepthOrStencilAspect())
+			{
+				RemainingExtras->BaseTexture->AllPlanesTrackedAccess[RemainingExtras->PlaneSlice] = RemainingExtras->PlaneAccess;
+			}
 
-template <> 
-void FTransitionProcessor<VkMemoryBarrier2, VkBufferMemoryBarrier2, VkImageMemoryBarrier2>::FinishBatch(ERHIPipeline SrcPipelines, ERHIPipeline DstPipelines, int32 BufferBatchStartIndex, int32 ImageBatchStartIndex)
-{
-	// :todo-jn: clean up
+			if ((RemainingBarriers->image != VK_NULL_HANDLE) && (RemainingBarriers->subresourceRange.aspectMask != 0))
+			{
+				const int32 RemainingCount = BarrierArrays.ImageExtraData.Num() - Index;
+				MergePlanes(RemainingBarriers, RemainingExtras, RemainingCount);
+			}
+		}
 
-	if (SrcPipelines == ERHIPipeline::AsyncCompute)
-	{
-		for (int32 Index = BufferBatchStartIndex; Index < BufferBarriers.Num(); ++Index)
+		// Merging Depth and Stencil transitions will also result in null aspectMask for the extra transition which needs to be removed.
+		for (int32 DstIndex = 0; DstIndex < BarrierArrays.ImageBarriers.Num();)
 		{
-			VkBufferMemoryBarrier2& BufferBarrier = BufferBarriers[Index];
-			BufferBarrier.srcStageMask &= Context.Device.GetComputeQueue()->GetSupportedStageBits();
-			BufferBarrier.srcAccessMask &= (VK_ACCESS_SHADER_READ_BIT | VK_ACCESS_SHADER_WRITE_BIT);
+			if ((BarrierArrays.ImageBarriers[DstIndex].image == VK_NULL_HANDLE) || ((BarrierArrays.ImageBarriers[DstIndex].subresourceRange.aspectMask == 0)))
+			{
+				BarrierArrays.ImageBarriers.RemoveAtSwap(DstIndex);
+			}
+			else
+			{
+				++DstIndex;
+			}
 		}
 
-		for (int32 Index = ImageBatchStartIndex; Index < ImageBarriers.Num(); ++Index)
+		// Submit
+		if (BarrierArrays.MemoryBarriers.Num() || BarrierArrays.BufferBarriers.Num() || BarrierArrays.ImageBarriers.Num())
 		{
-			VkImageMemoryBarrier2& Barrier = ImageBarriers[Index];
-			Barrier.srcStageMask &= Context.Device.GetComputeQueue()->GetSupportedStageBits();
-			Barrier.srcAccessMask &= (VK_ACCESS_SHADER_READ_BIT | VK_ACCESS_SHADER_WRITE_BIT);
+			// Submit merged stage masks with arrays of barriers
+			VulkanRHI::vkCmdPipelineBarrier(Context.GetCommandBuffer().GetHandle(), 
+				BarrierArrays.SrcStageMask, BarrierArrays.DstStageMask, 0 /*VkDependencyFlags*/,
+				BarrierArrays.MemoryBarriers.Num(), BarrierArrays.MemoryBarriers.GetData(),
+				BarrierArrays.BufferBarriers.Num(), BarrierArrays.BufferBarriers.GetData(),
+				BarrierArrays.ImageBarriers.Num(),  BarrierArrays.ImageBarriers.GetData());
 		}
 	}
+}
+
+
+
+
+
+
+
 
-	if (DstPipelines == ERHIPipeline::AsyncCompute)
+
+
+// Removes stages/access that aren't supported by the compute queue
+template<typename BarrierType>
+static void MaskSupportedAsyncFlags(FVulkanDevice& Device, TArray<BarrierType>& InOutBarriers, bool bMaskSrc, bool bMaskDst)
+{
+	const VkPipelineStageFlags SupportedComputeStageMask = Device.GetComputeQueue()->GetSupportedStageBits();
+	const VkAccessFlags SupportedComputeAccessMasks = Device.GetComputeQueue()->GetSupportedAccessFlags();
+
+	for (BarrierType& Barrier : InOutBarriers)
 	{
-		for (int32 Index = BufferBatchStartIndex; Index < BufferBarriers.Num(); ++Index)
+		if (bMaskSrc)
+		{
+			Barrier.srcStageMask &= SupportedComputeStageMask;
+			Barrier.srcAccessMask &= SupportedComputeAccessMasks;
+		}
+		if (bMaskDst)
 		{
-			VkBufferMemoryBarrier2& BufferBarrier = BufferBarriers[Index];
-			BufferBarrier.dstStageMask &= Context.Device.GetComputeQueue()->GetSupportedStageBits();
-			BufferBarrier.dstAccessMask &= (VK_ACCESS_SHADER_READ_BIT | VK_ACCESS_SHADER_WRITE_BIT);
+			Barrier.dstStageMask &= SupportedComputeStageMask;
+			Barrier.dstAccessMask &= SupportedComputeAccessMasks;
 		}
+	}
+}
+
+
+// Patches barriers for release/acquire of resources during queue ownership transfers
+template<typename BarrierType>
+static void PatchCrossPipeTransitions(TArray<BarrierType>& Barriers, FVulkanCommandListContext& Context, ERHIPipeline SrcPipelines, ERHIPipeline DstPipelines, bool bIsBeginTransition)
+{
+	const ERHIPipeline ExecutingPipeline = Context.GetPipeline();
+	const uint32 GraphicsFamilyIndex = Context.Device.GetGraphicsQueue()->GetFamilyIndex();
+	const uint32 ComputeFamilyIndex = Context.Device.GetComputeQueue()->GetFamilyIndex();
 
-		for (int32 Index = ImageBatchStartIndex; Index < ImageBarriers.Num(); ++Index)
+	// In the case where src and dst are both single pipelines, keep the layout changes to try to do all the work in a single barrier
+	if (IsSingleRHIPipeline(SrcPipelines) && IsSingleRHIPipeline(DstPipelines))
+	{
+		for (int32 Index = 0; Index < Barriers.Num(); ++Index)
 		{
-			VkImageMemoryBarrier2& Barrier = ImageBarriers[Index];
-			Barrier.dstStageMask &= Context.Device.GetComputeQueue()->GetSupportedStageBits();
-			Barrier.dstAccessMask &= (VK_ACCESS_SHADER_READ_BIT | VK_ACCESS_SHADER_WRITE_BIT);
+			BarrierType& Barrier = Barriers[Index];
+
+			Barrier.srcQueueFamilyIndex = (SrcPipelines == ERHIPipeline::Graphics) ? GraphicsFamilyIndex : ComputeFamilyIndex;
+			Barrier.dstQueueFamilyIndex = (DstPipelines == ERHIPipeline::Graphics) ? GraphicsFamilyIndex : ComputeFamilyIndex;
+
+			if (bIsBeginTransition)
+			{
+				// Release
+				check(SrcPipelines == ExecutingPipeline);
+				Barrier.dstAccessMask = 0;
+			}
+			else
+			{
+				// Acquire
+				check(DstPipelines == ExecutingPipeline);
+				Barrier.srcAccessMask = 0;
+			}
 		}
 	}
+	else // Src/Dst is ERHIPipeline::All, add a single queue transfer for now :todo-jn:
+	{
+		// The cross-pipe transition will have a Graphics source when it's Graphics->All or All->AsyncCompute
+		const bool bSrcIsGraphics = ((SrcPipelines == ERHIPipeline::Graphics) || (DstPipelines == ERHIPipeline::AsyncCompute));
+
+		// When a transition uses ERHIPipeline::All, it might include stages/access that isn't supported on all queues
+		MaskSupportedAsyncFlags(Context.Device, Barriers, !bSrcIsGraphics, bSrcIsGraphics);
+
+		for (int32 Index = 0; Index < Barriers.Num(); ++Index)
+		{
+			BarrierType& Barrier = Barriers[Index];
 
-	// Sync2 does not submit batches, only one big clump at the end
+			// Set the queue families and filter the stages for ERHIPipeline::All running on async compute
+			if (bSrcIsGraphics)
+			{
+				Barrier.srcQueueFamilyIndex = GraphicsFamilyIndex;
+				Barrier.dstQueueFamilyIndex = ComputeFamilyIndex;
+			}
+			else
+			{
+				Barrier.srcQueueFamilyIndex = ComputeFamilyIndex;
+				Barrier.dstQueueFamilyIndex = GraphicsFamilyIndex;
+			}
+
+			// Remove the layout change, it gets submitted separately (on the single pipeline for 1..N and N..1 transitions)
+			if constexpr (std::is_same_v<BarrierType, VkImageMemoryBarrier2>)
+			{
+				if (IsSingleRHIPipeline(SrcPipelines))
+				{
+					Barrier.oldLayout = Barrier.newLayout;
+				}
+				else if (IsSingleRHIPipeline(DstPipelines))
+				{
+					Barrier.newLayout = Barrier.oldLayout;
+				}
+			}
+
+			if (bIsBeginTransition)
+			{
+				// Release resource from current queue.
+				check(EnumHasAllFlags(SrcPipelines, ExecutingPipeline));
+				Barrier.dstAccessMask = 0;
+			}
+			else
+			{
+				// Acquire resource on current queue.
+				check(EnumHasAllFlags(DstPipelines, ExecutingPipeline));
+				Barrier.srcAccessMask = 0;
+			}
+		}
+	}
 }
 
+
+
+// Used to split up barrier batches in single calls to vkCmdPipelineBarrier2 on drivers with issues on larger batches
 template<typename BarrierType>
 static void SendBatchedBarriers(VkCommandBuffer CommandBuffer, VkDependencyInfo& BatchDependencyInfo, BarrierType*& BarrierPtr, uint32_t& BarrierCountRef, int32 TotalBarrierCount)
 {
@@ -1379,105 +1340,234 @@ static void SendBatchedBarriers(VkCommandBuffer CommandBuffer, VkDependencyInfo&
 	BarrierCountRef = 0;
 }
 
-template <>
-void FTransitionProcessor<VkMemoryBarrier2, VkBufferMemoryBarrier2, VkImageMemoryBarrier2>::FinishAll()
+
+
+// Create Vulkan barriers from RHI transitions when VK_KHR_Synchronization2 is supported
+void ProcessTransitionSync2(FVulkanCommandListContext& Context, TArrayView<const FRHITransition*>& Transitions, bool bIsBeginTransition)
 {
-	if ((MemoryBarriers.Num() == 0) && (BufferBarriers.Num() == 0) && (ImageBarriers.Num() == 0))
-	{
-		return;
-	}
+	auto SubmitBarriers = [bIsBeginTransition, &Context](TArrayView<const VkMemoryBarrier2> MemoryBarriers, TArrayView<const VkBufferMemoryBarrier2> BufferBarriers, TArrayView<const VkImageMemoryBarrier2> ImageBarriers, VkEvent BarrierEvent = VK_NULL_HANDLE)
+		{
+			if (MemoryBarriers.Num() || BufferBarriers.Num() || ImageBarriers.Num())
+			{
+				VkDependencyInfo DependencyInfo;
+				DependencyInfo.sType = VK_STRUCTURE_TYPE_DEPENDENCY_INFO;
+				DependencyInfo.pNext = nullptr;
+				DependencyInfo.dependencyFlags = 0;
+				DependencyInfo.memoryBarrierCount = MemoryBarriers.Num();
+				DependencyInfo.pMemoryBarriers = MemoryBarriers.GetData();
+				DependencyInfo.bufferMemoryBarrierCount = BufferBarriers.Num();
+				DependencyInfo.pBufferMemoryBarriers = BufferBarriers.GetData();
+				DependencyInfo.imageMemoryBarrierCount = ImageBarriers.Num();
+				DependencyInfo.pImageMemoryBarriers = ImageBarriers.GetData();
+
+				if (BarrierEvent)
+				{
+					if (bIsBeginTransition)
+					{
+						Context.GetCommandBuffer().BeginSplitBarrier(BarrierEvent, DependencyInfo);
+					}
+					else
+					{
+						Context.GetCommandBuffer().EndSplitBarrier(BarrierEvent, DependencyInfo);
+					}
+				}
+				else if ((GVulkanMaxBarriersPerBatch <= 0) || ((MemoryBarriers.Num() + BufferBarriers.Num() + ImageBarriers.Num()) < GVulkanMaxBarriersPerBatch))
+				{
+					VulkanRHI::vkCmdPipelineBarrier2KHR(Context.GetCommandBuffer().GetHandle(), &DependencyInfo);
+				}
+				else
+				{
+					VkDependencyInfo BatchDependencyInfo = DependencyInfo;
+					BatchDependencyInfo.memoryBarrierCount = 0;
+					BatchDependencyInfo.bufferMemoryBarrierCount = 0;
+					BatchDependencyInfo.imageMemoryBarrierCount = 0;
+
+					VkCommandBuffer CommandBuffer = Context.GetCommandBuffer().GetHandle();
+					SendBatchedBarriers(CommandBuffer, BatchDependencyInfo, BatchDependencyInfo.pMemoryBarriers, BatchDependencyInfo.memoryBarrierCount, DependencyInfo.memoryBarrierCount);
+					SendBatchedBarriers(CommandBuffer, BatchDependencyInfo, BatchDependencyInfo.pBufferMemoryBarriers, BatchDependencyInfo.bufferMemoryBarrierCount, DependencyInfo.bufferMemoryBarrierCount);
+					SendBatchedBarriers(CommandBuffer, BatchDependencyInfo, BatchDependencyInfo.pImageMemoryBarriers, BatchDependencyInfo.imageMemoryBarrierCount, DependencyInfo.imageMemoryBarrierCount);
+				}
+			}
+		};
 
-	VkDependencyInfo DependencyInfo;
-	DependencyInfo.sType = VK_STRUCTURE_TYPE_DEPENDENCY_INFO;
-	DependencyInfo.pNext = nullptr;
-	DependencyInfo.dependencyFlags = 0;
-	DependencyInfo.memoryBarrierCount = MemoryBarriers.Num();
-	DependencyInfo.pMemoryBarriers = MemoryBarriers.GetData();
-	DependencyInfo.bufferMemoryBarrierCount = BufferBarriers.Num();
-	DependencyInfo.pBufferMemoryBarriers = BufferBarriers.GetData();
-	DependencyInfo.imageMemoryBarrierCount = ImageBarriers.Num();
-	DependencyInfo.pImageMemoryBarriers = ImageBarriers.GetData();
 
-	if ((GVulkanMaxBarriersPerBatch <= 0) || ((MemoryBarriers.Num() + BufferBarriers.Num() + ImageBarriers.Num()) < GVulkanMaxBarriersPerBatch))
-	{
-		VulkanRHI::vkCmdPipelineBarrier2KHR(Context.GetCommandBuffer().GetHandle(), &DependencyInfo);
-	}
-	else
+	TArray<VkBufferMemoryBarrier2> TempBufferBarriers;
+	TArray<VkImageMemoryBarrier2> TempImageBarriers;
+
+	const bool bUseOwnershipTransfers = Context.Device.HasAsyncComputeQueue() && (!GVulkanAllowConcurrentBuffer || !GVulkanAllowConcurrentImage);
+	for (const FRHITransition* Transition : Transitions)
 	{
-		VkDependencyInfo BatchDependencyInfo = DependencyInfo;
-		BatchDependencyInfo.memoryBarrierCount = 0;
-		BatchDependencyInfo.bufferMemoryBarrierCount = 0;
-		BatchDependencyInfo.imageMemoryBarrierCount = 0;
+		const FVulkanTransitionData* Data = Transition->GetPrivateData<FVulkanTransitionData>();
 
-		VkCommandBuffer CommandBuffer = Context.GetCommandBuffer().GetHandle();
-		SendBatchedBarriers(CommandBuffer, BatchDependencyInfo, BatchDependencyInfo.pMemoryBarriers, BatchDependencyInfo.memoryBarrierCount, DependencyInfo.memoryBarrierCount);
-		SendBatchedBarriers(CommandBuffer, BatchDependencyInfo, BatchDependencyInfo.pBufferMemoryBarriers, BatchDependencyInfo.bufferMemoryBarrierCount, DependencyInfo.bufferMemoryBarrierCount);
-		SendBatchedBarriers(CommandBuffer, BatchDependencyInfo, BatchDependencyInfo.pImageMemoryBarriers, BatchDependencyInfo.imageMemoryBarrierCount, DependencyInfo.imageMemoryBarrierCount);
-	}
-}
+		const bool bIsSingleQueue = IsSingleRHIPipeline(Data->SrcPipelines) && (Data->SrcPipelines == Data->DstPipelines);
+		const ERHIPipeline TargetPipeline = bIsBeginTransition ? Data->SrcPipelines : Data->DstPipelines;
+		const ERHIPipeline OtherPipeline = bIsBeginTransition ? Data->DstPipelines : Data->SrcPipelines;
+		const ERHIPipeline ExecutingPipeline = Context.GetPipeline();
 
-template <>
-void FTransitionProcessor<VkMemoryBarrier2, VkBufferMemoryBarrier2, VkImageMemoryBarrier2>::AddMemoryBarrier(const VkMemoryBarrier2& Barrier)
-{
-	if (MemoryBarriers.Num() == 0)
-	{
-		VkMemoryBarrier2& NewBarrier = MemoryBarriers.AddDefaulted_GetRef();
-		ZeroVulkanStruct(NewBarrier, VK_STRUCTURE_TYPE_MEMORY_BARRIER_2);
-	}
+		checkf(EnumHasAnyFlags(TargetPipeline, Context.GetPipeline()),
+			TEXT("The %s pipelines for this transition are [%s], but it's submitted on the [%s] queue."),
+			bIsBeginTransition ? TEXT("SRC") : TEXT("DST"),
+			*GetRHIPipelineName(TargetPipeline),
+			*GetRHIPipelineName(Context.GetPipeline())
+		);
 
-	// Mash everything into a single barrier
-	VkMemoryBarrier2& MemoryBarrier = MemoryBarriers[0];
+		// Single queue barrier that aren't split only submit in EndTransition
+		if (bIsSingleQueue && bIsBeginTransition && !Data->EventHandle)
+		{
+			continue;
+		}
 
-	MergeBarrierAccessMask(MemoryBarrier, Barrier.srcAccessMask, Barrier.dstAccessMask);
-	MemoryBarrier.srcStageMask |= Barrier.srcStageMask;
-	MemoryBarrier.dstStageMask |= Barrier.dstStageMask;
-}
+		FSync2BarrierArrays BarrierArrays;
+		ConvertTransitionToBarriers(Context, *Data, BarrierArrays);
 
-template <>
-void FTransitionProcessor<VkMemoryBarrier2, VkBufferMemoryBarrier2, VkImageMemoryBarrier2>::AddBufferBarrier(const VkBufferMemoryBarrier2& Barrier)
-{
-	BufferBarriers.Emplace(Barrier);
-}
+		// Submit split-barriers right away (they are always single queue)
+		if (Data->EventHandle)
+		{
+			checkf(bIsSingleQueue, TEXT("Split barriers must remain on same queue!"));
+			SubmitBarriers(MakeArrayView(BarrierArrays.MemoryBarriers), MakeArrayView(BarrierArrays.BufferBarriers), MakeArrayView(BarrierArrays.ImageBarriers), Data->EventHandle);
+			continue;
+		}
 
-template <>
-void FTransitionProcessor<VkMemoryBarrier2, VkBufferMemoryBarrier2, VkImageMemoryBarrier2>::AddImageBarrier(const VkImageMemoryBarrier2& Barrier, const FImageBarrierExtraData& ExtraData)
-{
-	ImageBarriers.Emplace(Barrier);
-	ImageBarrierExtras.Emplace(ExtraData);
-}
+		// Same queue, or single-queue to single-queue transfers, can be submitted directly without copies
+		if (IsSingleRHIPipeline(Data->SrcPipelines) && IsSingleRHIPipeline(Data->DstPipelines))
+		{
+			// For cross-queue 1..1 transitions we will keep the layout change in the same barrier as the queue transfer
+			if (bUseOwnershipTransfers && (Data->SrcPipelines != Data->DstPipelines))
+			{
+				if (!GVulkanAllowConcurrentBuffer)
+				{
+					PatchCrossPipeTransitions(BarrierArrays.BufferBarriers, Context, Data->SrcPipelines, Data->DstPipelines, bIsBeginTransition);
+				}
 
-template <>
-void FTransitionProcessor<VkMemoryBarrier2, VkBufferMemoryBarrier2, VkImageMemoryBarrier2>::ApplyTracking(uint32 TargetIndex, const FImageBarrierExtraData* RemainingExtras, int32 RemainingCount, bool IsCrossPipe)
-{
-	// Nothing to do when VK_KHR_separate_depth_stencil_layouts and VK_KHR_synchronization2 are supported (requirements for SM6).
-}
+				if (!GVulkanAllowConcurrentImage)
+				{
+					PatchCrossPipeTransitions(BarrierArrays.ImageBarriers, Context, Data->SrcPipelines, Data->DstPipelines, bIsBeginTransition);
+				}
+			}
 
+			SubmitBarriers(MakeArrayView(BarrierArrays.MemoryBarriers), MakeArrayView(BarrierArrays.BufferBarriers), MakeArrayView(BarrierArrays.ImageBarriers), Data->EventHandle);
+			continue;
+		}
 
+		// For 1..N or N..1 transitions we submit the barriers for layout changes on the single pipeline (rest is covered by the sema)
+		const bool bNeedsPreLayoutChange = (bIsBeginTransition && ((IsSingleRHIPipeline(Data->SrcPipelines) && !IsSingleRHIPipeline(Data->DstPipelines) && (Data->SrcPipelines == ExecutingPipeline))));
+		const bool bNeedsRelease = (bIsBeginTransition && (bNeedsPreLayoutChange ||
+			(IsSingleRHIPipeline(Data->DstPipelines) && !IsSingleRHIPipeline(Data->SrcPipelines) && (Data->DstPipelines != ExecutingPipeline))));
+		const bool bNeedsPostLayoutChange = (!bIsBeginTransition && (IsSingleRHIPipeline(Data->DstPipelines) && !IsSingleRHIPipeline(Data->SrcPipelines) && (Data->DstPipelines == ExecutingPipeline)));
+		const bool bNeedsAcquire = (!bIsBeginTransition && (bNeedsPostLayoutChange ||
+			(IsSingleRHIPipeline(Data->SrcPipelines) && !IsSingleRHIPipeline(Data->DstPipelines) && (Data->SrcPipelines != ExecutingPipeline))));
 
 
+		// For resources without concurrent sharing mode, make copies of the array so we can:
+		// - wipe the layout transitions from the barriers
+		// - patch in the actual queue family ownership transfer
+		if (bUseOwnershipTransfers && (bNeedsRelease || bNeedsAcquire))
+		{
+			if (!GVulkanAllowConcurrentBuffer)
+			{
+				TempBufferBarriers = BarrierArrays.BufferBarriers;
+				PatchCrossPipeTransitions(TempBufferBarriers, Context, Data->SrcPipelines, Data->DstPipelines, bIsBeginTransition);
+			}
 
-static TArray<VulkanRHI::FSemaphore*> ExtractTransitionSemaphores(TArrayView<const FRHITransition*>& Transitions, ERHIPipeline SkipSrcPipeline = ERHIPipeline::None)
-{
-	TArray<VulkanRHI::FSemaphore*> Semaphores;
+			if (!GVulkanAllowConcurrentImage)
+			{
+				TempImageBarriers = BarrierArrays.ImageBarriers;
+				PatchCrossPipeTransitions(TempImageBarriers, Context, Data->SrcPipelines, Data->DstPipelines, bIsBeginTransition);
+			}
+		}
 
-	for (const FRHITransition* Transition : Transitions)
-	{
-		const FVulkanTransitionData* Data = Transition->GetPrivateData<FVulkanTransitionData>();
-		if (Data->Semaphore == nullptr)
+		if (bNeedsPreLayoutChange)
 		{
-			continue;
+			// Remove unsupported flags if we're submitting on the compute queue
+			if (ExecutingPipeline == ERHIPipeline::AsyncCompute)
+			{
+				MaskSupportedAsyncFlags(Context.Device, BarrierArrays.BufferBarriers, false, true);
+				MaskSupportedAsyncFlags(Context.Device, BarrierArrays.ImageBarriers, false, true);
+			}
+
+			SubmitBarriers(MakeArrayView(BarrierArrays.MemoryBarriers), MakeArrayView(BarrierArrays.BufferBarriers), MakeArrayView(BarrierArrays.ImageBarriers));
 		}
 
-		// Avoid unnecessary semaphores when DstPipelines is ERHIPipeline::All
-		if ((SkipSrcPipeline != ERHIPipeline::None) && EnumHasAllFlags(Data->DstPipelines, Data->SrcPipelines) && (Data->SrcPipelines == SkipSrcPipeline))
+		if (TempBufferBarriers.Num() || TempImageBarriers.Num())
 		{
-			continue;
+			SubmitBarriers(MakeArrayView(BarrierArrays.MemoryBarriers), MakeArrayView(TempBufferBarriers), MakeArrayView(TempImageBarriers));
+			TempBufferBarriers.Reset();
+			TempImageBarriers.Reset();
 		}
 
-		Semaphores.Add(Data->Semaphore);
+		if (bNeedsPostLayoutChange)
+		{
+			// Remove unsupported flags if we're submitting on the compute queue
+			if (ExecutingPipeline == ERHIPipeline::AsyncCompute)
+			{
+				MaskSupportedAsyncFlags(Context.Device, BarrierArrays.BufferBarriers, true, false);
+				MaskSupportedAsyncFlags(Context.Device, BarrierArrays.ImageBarriers, true, false);
+			}
+
+			SubmitBarriers(MakeArrayView(BarrierArrays.MemoryBarriers), MakeArrayView(BarrierArrays.BufferBarriers), MakeArrayView(BarrierArrays.ImageBarriers));
+		}
 	}
+}
+
+
+
+
+
+
+
+
+
+
+
+
+
+
+
+
+
+
+
+
+
+
 
-	return Semaphores;
+
+
+
+
+
+
+
+
+
+
+
+static void ProcessCrossQueueSemaphores(FVulkanCommandListContext& Context, TArrayView<const FRHITransition*>& Transitions, bool bIsBeginTransition)
+{
+	for (const FRHITransition* Transition : Transitions)
+	{
+		const FVulkanTransitionData* Data = Transition->GetPrivateData<FVulkanTransitionData>();
+		if (Data->Semaphore)
+		{
+			check(Data->SrcPipelines != Data->DstPipelines);
+			if (bIsBeginTransition)
+			{
+				if ((IsSingleRHIPipeline(Data->SrcPipelines) && (Data->SrcPipelines == Context.GetPipeline())) ||
+					(IsSingleRHIPipeline(Data->DstPipelines) && (Data->DstPipelines != Context.GetPipeline())))
+				{
+					Context.AddSignalSemaphore(Data->Semaphore);
+				}
+			}
+			else
+			{
+				if ((IsSingleRHIPipeline(Data->DstPipelines) && (Data->DstPipelines == Context.GetPipeline())) ||
+					(IsSingleRHIPipeline(Data->SrcPipelines) && (Data->SrcPipelines != Context.GetPipeline())))
+				{
+					Context.AddWaitSemaphore(VK_PIPELINE_STAGE_ALL_COMMANDS_BIT, Data->Semaphore);
+				}
+			}
+		}
+	}
 }
 
 
@@ -1485,43 +1575,49 @@ void FVulkanCommandListContext::RHIBeginTransitions(TArrayView<const FRHITransit
 {
 	if (Device.SupportsParallelRendering())
 	{
-		FTransitionProcessor<VkMemoryBarrier2, VkBufferMemoryBarrier2, VkImageMemoryBarrier2> Processor(*this, true);
-		Processor.Process(Transitions);
+		constexpr bool bIsBeginTransition = true;
+
+		ProcessTransitionSync2(*this, Transitions, bIsBeginTransition);
+
+		// Signal semaphores after the release barriers
+		ProcessCrossQueueSemaphores(*this, Transitions, bIsBeginTransition);
 	}
 	else
 	{
-		FTransitionProcessor<VkMemoryBarrier, VkBufferMemoryBarrier, VkImageMemoryBarrier> Processor(*this, true);
-		Processor.Process(Transitions);
+		// Nothing to do for legacy barriers on begin (no split support, no async compute support)
 	}
 
-	// Signal semaphores
-	TArray<VulkanRHI::FSemaphore*> SignalSemaphores = ExtractTransitionSemaphores(Transitions);
-	if (SignalSemaphores.Num() > 0)
-	{
-		AddSignalSemaphores(SignalSemaphores);
-	}
 }
 
 void FVulkanCommandListContext::RHIEndTransitions(TArrayView<const FRHITransition*> Transitions)
 {
-	TArray<VulkanRHI::FSemaphore*> WaitSemaphores = ExtractTransitionSemaphores(Transitions, GetPipeline());
-	if (WaitSemaphores.Num() > 0)
-	{
-		AddWaitSemaphores(VK_PIPELINE_STAGE_TOP_OF_PIPE_BIT, WaitSemaphores);
-	}
-
 	if (Device.SupportsParallelRendering())
 	{
-		FTransitionProcessor<VkMemoryBarrier2, VkBufferMemoryBarrier2, VkImageMemoryBarrier2> Processor(*this, false);
-		Processor.Process(Transitions);
+		constexpr bool bIsBeginTransition = false;
+
+		// Wait on semaphores before the acquire barriers
+		ProcessCrossQueueSemaphores(*this, Transitions, bIsBeginTransition);
+
+		ProcessTransitionSync2(*this, Transitions, bIsBeginTransition);
 	}
 	else
 	{
-		FTransitionProcessor<VkMemoryBarrier, VkBufferMemoryBarrier, VkImageMemoryBarrier> Processor(*this, false);
-		Processor.Process(Transitions);
+		ProcessTransitionLegacy(*this, Transitions);
 	}
 }
 
+
+
+
+
+
+
+
+
+
+
+
+
 void FVulkanPipelineBarrier::AddMemoryBarrier(VkAccessFlags InSrcAccessFlags, VkAccessFlags InDstAccessFlags, VkPipelineStageFlags InSrcStageMask, VkPipelineStageFlags InDstStageMask)
 {
 	if (MemoryBarriers.Num() == 0)
diff --git a/Engine/Source/Runtime/VulkanRHI/Private/VulkanDevice.cpp b/Engine/Source/Runtime/VulkanRHI/Private/VulkanDevice.cpp
--- a/Engine/Source/Runtime/VulkanRHI/Private/VulkanDevice.cpp
+++ b/Engine/Source/Runtime/VulkanRHI/Private/VulkanDevice.cpp
@@ -81,22 +81,14 @@ TAutoConsoleVariable<int32> CVarDelayAcquireBackBuffer(
 	1,
 	TEXT("Whether to delay acquiring the back buffer \n")
 	TEXT(" 0: acquire next image on frame start \n")
-	TEXT(" 1: acquire next image just before presenting, rendering is done to intermediate image which is then copied to a real backbuffer (default) \n")
-	TEXT(" 2: acquire next image on first use"),
+	TEXT(" 1: acquire next image just before presenting, rendering is done to intermediate image which is then copied to a real backbuffer (default) \n"),
 	ECVF_ReadOnly
 );
 
 static EDelayAcquireImageType DelayAcquireBackBuffer()
 {
-	int32 DelayType = CVarDelayAcquireBackBuffer.GetValueOnAnyThread();
-	switch (DelayType)
-	{
-	case 1:
-		return EDelayAcquireImageType::DelayAcquire;
-	case 2:
-		return EDelayAcquireImageType::LazyAcquire;
-	}
-	return EDelayAcquireImageType::None;
+	const int32 DelayType = CVarDelayAcquireBackBuffer.GetValueOnAnyThread();
+	return (DelayType == 0) ? EDelayAcquireImageType::None : EDelayAcquireImageType::DelayAcquire;
 }
 
 #if VULKAN_SUPPORTS_VALIDATION_CACHE
diff --git a/Engine/Source/Runtime/VulkanRHI/Private/VulkanDevice.h b/Engine/Source/Runtime/VulkanRHI/Private/VulkanDevice.h
--- a/Engine/Source/Runtime/VulkanRHI/Private/VulkanDevice.h
+++ b/Engine/Source/Runtime/VulkanRHI/Private/VulkanDevice.h
@@ -525,7 +525,7 @@ public:
 
 	inline bool NeedsAllPlanes() const
 	{
-		return !OptionalDeviceExtensions.HasSeparateDepthStencilLayouts;
+		return !SupportsParallelRendering();
 	}
 
 	inline bool SupportsParallelRendering() const
diff --git a/Engine/Source/Runtime/VulkanRHI/Private/VulkanMemory.h b/Engine/Source/Runtime/VulkanRHI/Private/VulkanMemory.h
--- a/Engine/Source/Runtime/VulkanRHI/Private/VulkanMemory.h
+++ b/Engine/Source/Runtime/VulkanRHI/Private/VulkanMemory.h
@@ -119,7 +119,6 @@ enum class EDelayAcquireImageType
 {
 	None,			// acquire next image on frame start
 	DelayAcquire,	// acquire next image just before presenting, rendering is done to intermediate image which is copied to real backbuffer
-	LazyAcquire,	// acquire next image on first use
 };
 
 extern EDelayAcquireImageType GVulkanDelayAcquireImage;
diff --git a/Engine/Source/Runtime/VulkanRHI/Private/VulkanRHI.cpp b/Engine/Source/Runtime/VulkanRHI/Private/VulkanRHI.cpp
--- a/Engine/Source/Runtime/VulkanRHI/Private/VulkanRHI.cpp
+++ b/Engine/Source/Runtime/VulkanRHI/Private/VulkanRHI.cpp
@@ -900,7 +900,7 @@ void FVulkanDynamicRHI::InitInstance()
 		GRHISupportsParallelRHIExecute = false;
 #else
 		GRHISupportsRHIThread = GRHIThreadCvar->GetInt() != 0;
-		GRHISupportsParallelRHIExecute = GRHIThreadCvar->GetInt() > 1;
+		GRHISupportsParallelRHIExecute = Device->SupportsParallelRendering() && (GRHIThreadCvar->GetInt() > 1);
 #endif
 
 		GRHISupportsParallelRenderPasses = GRHISupportsParallelRHIExecute;
diff --git a/Engine/Source/Runtime/VulkanRHI/Private/VulkanResources.h b/Engine/Source/Runtime/VulkanRHI/Private/VulkanResources.h
--- a/Engine/Source/Runtime/VulkanRHI/Private/VulkanResources.h
+++ b/Engine/Source/Runtime/VulkanRHI/Private/VulkanResources.h
@@ -613,8 +613,6 @@ public:
 
 	FTextureRHIRef AliasedTexture;
 
-	virtual void OnLayoutTransition(FVulkanCommandListContext& Context, VkImageLayout NewLayout) {}
-
 	template<typename T>
 	void DumpMemory(T Callback)
 	{
diff --git a/Engine/Source/Runtime/VulkanRHI/Private/VulkanViewport.cpp b/Engine/Source/Runtime/VulkanRHI/Private/VulkanViewport.cpp
--- a/Engine/Source/Runtime/VulkanRHI/Private/VulkanViewport.cpp
+++ b/Engine/Source/Runtime/VulkanRHI/Private/VulkanViewport.cpp
@@ -57,14 +57,6 @@ void FVulkanBackBuffer::OnAdvanceBackBufferFrame(FRHICommandListImmediate& RHICm
 	ReleaseAcquiredImage();
 }
 
-void FVulkanBackBuffer::OnLayoutTransition(FVulkanCommandListContext& Context, VkImageLayout NewLayout)
-{
-	if (GVulkanDelayAcquireImage == EDelayAcquireImageType::LazyAcquire)
-	{
-		AcquireBackBufferImage(Context);
-	}
-}
-
 void FVulkanBackBuffer::AcquireBackBufferImage(FVulkanCommandListContext& Context)
 {
 	check(Viewport);
diff --git a/Engine/Source/Runtime/VulkanRHI/Private/VulkanViewport.h b/Engine/Source/Runtime/VulkanRHI/Private/VulkanViewport.h
--- a/Engine/Source/Runtime/VulkanRHI/Private/VulkanViewport.h
+++ b/Engine/Source/Runtime/VulkanRHI/Private/VulkanViewport.h
@@ -25,8 +25,6 @@ public:
 	FVulkanBackBuffer(FVulkanDevice& Device, FVulkanViewport* InViewport, EPixelFormat Format, uint32 SizeX, uint32 SizeY, ETextureCreateFlags UEFlags);
 	virtual ~FVulkanBackBuffer();
 	
-	void OnLayoutTransition(FVulkanCommandListContext& Context, VkImageLayout NewLayout) override final;
-
 	void OnGetBackBufferImage(FRHICommandListImmediate& RHICmdList);
 	void OnAdvanceBackBufferFrame(FRHICommandListImmediate& RHICmdList);
 
