VulkanRHI: Add missing code for split barrier creation.

#rnx

diff --git a/Engine/Source/Runtime/VulkanRHI/Private/VulkanBarriers.cpp b/Engine/Source/Runtime/VulkanRHI/Private/VulkanBarriers.cpp
--- a/Engine/Source/Runtime/VulkanRHI/Private/VulkanBarriers.cpp
+++ b/Engine/Source/Runtime/VulkanRHI/Private/VulkanBarriers.cpp
@@ -661,6 +661,25 @@ void FVulkanDynamicRHI::RHICreateTransition(FRHITransition* Transition, const FR
 	{
 		Data->Semaphore = new VulkanRHI::FSemaphore(*Device);
 	}
+
+	// If we're staying on the same queue, use split barriers if they are permitted and supported
+	if (GVulkanAllowSplitBarriers && (GVulkanMaxBarriersPerBatch <= 0) && Device->SupportsParallelRendering() &&
+		(SrcPipelines == DstPipelines) && !EnumHasAnyFlags(CreateInfo.Flags, ERHITransitionCreateFlags::NoSplit))
+	{
+		// Track if host stage is used, it will prevent using split barrier
+		bool bIncludesHostSync = false;
+		for (const FRHITransitionInfo& Info : Data->TransitionInfos)
+		{
+			if (EnumHasAnyFlags(Info.AccessBefore, ERHIAccess::CPURead) || EnumHasAnyFlags(Info.AccessAfter, ERHIAccess::CPURead))
+			{
+				bIncludesHostSync = true;
+				break;
+			}
+		}
+
+		// Create an event for the split barriers
+		Data->EventHandle = !bIncludesHostSync ? Device->GetBarrierEvent() : VK_NULL_HANDLE;
+	}
 }
 
 void FVulkanDynamicRHI::RHIReleaseTransition(FRHITransition* Transition)
