VulkanRHI: Fix missing stages on Sync2 memory barriers.

#rnx

diff --git a/Engine/Source/Runtime/VulkanRHI/Private/VulkanBarriers.cpp b/Engine/Source/Runtime/VulkanRHI/Private/VulkanBarriers.cpp
--- a/Engine/Source/Runtime/VulkanRHI/Private/VulkanBarriers.cpp
+++ b/Engine/Source/Runtime/VulkanRHI/Private/VulkanBarriers.cpp
@@ -908,8 +908,8 @@ void ConvertTransitionToBarriers(FVulkanCommandListContext& Context, const FVulk
 
 			if constexpr (!bIsLegacyBarrier)
 			{
-				MemoryBarrier.srcStageMask = SrcStageMask;
-				MemoryBarrier.dstStageMask = DstStageMask;
+				MemoryBarrier.srcStageMask |= SrcStageMask;
+				MemoryBarrier.dstStageMask |= DstStageMask;
 			}
 		}
 		else if (Buffer != nullptr)
