VulkanRHI: Add a clear message for VK_ERROR_LAYER_NOT_PRESENT with the option of attempting to load without layers.

#rnx

diff --git a/Engine/Source/Runtime/VulkanRHI/Private/VulkanRHI.cpp b/Engine/Source/Runtime/VulkanRHI/Private/VulkanRHI.cpp
--- a/Engine/Source/Runtime/VulkanRHI/Private/VulkanRHI.cpp
+++ b/Engine/Source/Runtime/VulkanRHI/Private/VulkanRHI.cpp
@@ -660,6 +660,34 @@ void FVulkanDynamicRHI::CreateInstance()
 
 	VkResult Result = VulkanRHI::vkCreateInstance(&InstInfo, VULKAN_CPU_ALLOCATOR, &Instance);
 	
+	if (Result == VK_ERROR_LAYER_NOT_PRESENT)
+	{
+		UE_LOG(LogVulkanRHI, Warning, TEXT("Vulkan instance creation returned an error with the requested layers (%d):"), InstanceLayers.Num());
+
+		for (const ANSICHAR* AnsiLayerName : InstanceLayers)
+		{
+			const FString LayerStr = ANSI_TO_TCHAR(AnsiLayerName);
+			UE_LOG(LogVulkanRHI, Warning, TEXT("- %s"), *LayerStr);
+		}
+
+		const EAppReturnType::Type MsgBoxResult = FPlatformMisc::MessageBoxExt(EAppMsgType::YesNo, TEXT(
+			"ERROR: Vulkan driver couldn't load one of the requested layers (see log for details).\n\n"
+			"Retry without layers?"),
+			TEXT("Incompatible Vulkan layer found!"));
+
+		if (MsgBoxResult == EAppReturnType::Yes)
+		{
+			InstInfo.enabledLayerCount = 0;
+			Result = VulkanRHI::vkCreateInstance(&InstInfo, VULKAN_CPU_ALLOCATOR, &Instance);
+		}
+		else
+		{
+			FPlatformMisc::RequestExitWithStatus(true, 1);
+			// unreachable
+			return;
+		}
+	}
+
 	FVulkanPlatform::NotifyFoundInstanceLayersAndExtensions(InstanceLayers, InstanceExtensions);
 
 	if (Result == VK_ERROR_INCOMPATIBLE_DRIVER)
@@ -713,7 +741,7 @@ void FVulkanDynamicRHI::CreateInstance()
 	else if (Result != VK_SUCCESS)
 	{
 		FPlatformMisc::MessageBoxExt(EAppMsgType::Ok, TEXT(
-			"Vulkan failed to create instance (apiVersion=0x%x)\n\nDo you have a compatible Vulkan "
+			"Vulkan failed to create instance.\n\nDo you have a compatible Vulkan "
 			 "driver (ICD) installed?\nPlease look at "
 			 "the Getting Started guide for additional information."), TEXT("No Vulkan driver found!"));
 		FPlatformMisc::RequestExitWithStatus(true, 1);
