VulkanRHI: Fix unknown initial buffer state in ray tracing as compaction

#rnx

diff --git a/Engine/Source/Runtime/VulkanRHI/Private/VulkanRayTracing.cpp b/Engine/Source/Runtime/VulkanRHI/Private/VulkanRayTracing.cpp
--- a/Engine/Source/Runtime/VulkanRHI/Private/VulkanRayTracing.cpp
+++ b/Engine/Source/Runtime/VulkanRHI/Private/VulkanRayTracing.cpp
@@ -506,7 +506,8 @@ void FVulkanRayTracingGeometry::CompactAccelerationStructure(FVulkanCommandBuffe
 
 	const FString DebugNameString = Initializer.DebugName.ToString();
 	const FRHIBufferCreateDesc CreateDesc =
-		FRHIBufferCreateDesc::Create(*DebugNameString, InSizeAfterCompaction, 0, BUF_AccelerationStructure);
+		FRHIBufferCreateDesc::Create(*DebugNameString, InSizeAfterCompaction, 0, BUF_AccelerationStructure)
+		.SetInitialState(ERHIAccess::BVHWrite);
 
 	AccelerationStructureBuffer = new FVulkanBuffer(Device, CreateDesc);
 
