VulkanRHI: Add CVar to control transfer queue support.  Make ForEachQueue only apply to active queues.  Do not create FVulkanQueue for unused queues.  Remove unused code.  Log timestamp period in logs on startup.

#rnx

diff --git a/Engine/Source/Runtime/VulkanRHI/Private/VulkanDevice.cpp b/Engine/Source/Runtime/VulkanRHI/Private/VulkanDevice.cpp
--- a/Engine/Source/Runtime/VulkanRHI/Private/VulkanDevice.cpp
+++ b/Engine/Source/Runtime/VulkanRHI/Private/VulkanDevice.cpp
@@ -29,6 +29,14 @@ static TAutoConsoleVariable<int32> GRHIAllowAsyncComputeCvar(
 	ECVF_ReadOnly
 );
 
+static TAutoConsoleVariable<int32> GRHIAllowTransferQueueCvar(
+	TEXT("r.Vulkan.AllowTransferQueue"),
+	0,
+	TEXT("0 to disable transfer queue (if available)\n")
+	TEXT("1 to allow transfer queue"),
+	ECVF_ReadOnly
+);
+
 int32 GVulkanAllowConcurrentBuffer = 1;
 static TAutoConsoleVariable<int32> GCVarAllowConcurrentBuffer(
 	TEXT("r.Vulkan.AllowConcurrentBuffer"),
@@ -250,8 +258,9 @@ FVulkanDevice::FVulkanDevice(FVulkanDynamicRHI* InRHI, VkPhysicalDevice InGpu)
 	UE_LOG(LogVulkanRHI, Display, TEXT("- DeviceName: %s"), ANSI_TO_TCHAR(GpuProps.deviceName));
 	UE_LOG(LogVulkanRHI, Display, TEXT("- API=%d.%d.%d (0x%x) Driver=0x%x VendorId=0x%x"), VK_VERSION_MAJOR(GpuProps.apiVersion), VK_VERSION_MINOR(GpuProps.apiVersion), VK_VERSION_PATCH(GpuProps.apiVersion), GpuProps.apiVersion, GpuProps.driverVersion, GpuProps.vendorID);
 	UE_LOG(LogVulkanRHI, Display, TEXT("- DeviceID=0x%x Type=%s"), GpuProps.deviceID, VK_TYPE_TO_STRING(VkPhysicalDeviceType, GpuProps.deviceType));
-	UE_LOG(LogVulkanRHI, Display, TEXT("- Max Descriptor Sets Bound %d, Timestamps %d"), GpuProps.limits.maxBoundDescriptorSets, GpuProps.limits.timestampComputeAndGraphics);
-	UE_LOG(LogVulkanRHI, Display, TEXT("- Calibrated Timestamp Domain: %s."), VK_TYPE_TO_STRING(VkTimeDomainKHR, FVulkanPlatform::GetTimeDomain()));
+	UE_LOG(LogVulkanRHI, Display, TEXT("- Max Descriptor Sets Bound %d"), GpuProps.limits.maxBoundDescriptorSets);
+	UE_LOG(LogVulkanRHI, Display, TEXT("- Timestamps:  ComputeAndGraphics=%d  Domain=%s  Period=%f"), 
+		GpuProps.limits.timestampComputeAndGraphics, VK_TYPE_TO_STRING(VkTimeDomainKHR, FVulkanPlatform::GetTimeDomain()), GpuProps.limits.timestampPeriod);
 
 	ensureMsgf(VendorId != EGpuVendorId::Unknown, TEXT("Unknown vendor ID 0x%x"), GpuProps.vendorID);
 }
@@ -359,7 +368,10 @@ void FVulkanDevice::CreateDevice(TArray<const ANSICHAR*>& DeviceLayers, FVulkanD
 		if (VKHasAllFlags(CurrProps.queueFlags, VK_QUEUE_TRANSFER_BIT))
 		{
 			// Prefer a non-gfx transfer queue
-			if (TransferQueueFamilyIndex == -1 && !VKHasAnyFlags(CurrProps.queueFlags, VK_QUEUE_GRAPHICS_BIT | VK_QUEUE_COMPUTE_BIT))
+			if ((TransferQueueFamilyIndex == -1) &&
+				!VKHasAnyFlags(CurrProps.queueFlags, VK_QUEUE_GRAPHICS_BIT | VK_QUEUE_COMPUTE_BIT) &&
+				(GRHIAllowTransferQueueCvar.GetValueOnAnyThread() != 0) &&
+				SupportsParallelRendering())
 			{
 				TransferQueueFamilyIndex = FamilyIndex;
 				bIsValidQueue = true;
@@ -411,45 +423,34 @@ void FVulkanDevice::CreateDevice(TArray<const ANSICHAR*>& DeviceLayers, FVulkanD
 
 	ActiveQueueFamilies.Reserve((int32)EVulkanQueueType::Count);
 
+	const uint32 NumBits = QueueFamilyProps[GfxQueueFamilyIndex].timestampValidBits;
+	UE_LOG(LogVulkanRHI, Display, TEXT("TimestampValidBits=%u"), NumBits);
+
 	// Create Graphics Queue, here we submit command buffers for execution
 	Queues[(int32)EVulkanQueueType::Graphics] = new FVulkanQueue(*this, GfxQueueFamilyIndex, EVulkanQueueType::Graphics);
 	ActiveQueueFamilies.Add(GfxQueueFamilyIndex);
 	if (ComputeQueueFamilyIndex == -1)
 	{
-		// If we didn't find a dedicated Queue, use the default one
-		ComputeQueueFamilyIndex = GfxQueueFamilyIndex;
+		// If we didn't find a dedicated Queue, leave it null
+		Queues[(int32)EVulkanQueueType::AsyncCompute] = nullptr;
 	}
 	else
 	{
-		bAsyncComputeQueue = true;
+		ensure(NumBits == QueueFamilyProps[ComputeQueueFamilyIndex].timestampValidBits);
+		Queues[(int32)EVulkanQueueType::AsyncCompute] = new FVulkanQueue(*this, ComputeQueueFamilyIndex, EVulkanQueueType::AsyncCompute);
 		ActiveQueueFamilies.Add(ComputeQueueFamilyIndex);
 	}
-	Queues[(int32)EVulkanQueueType::AsyncCompute] = new FVulkanQueue(*this, ComputeQueueFamilyIndex, EVulkanQueueType::AsyncCompute);
+
 	if (TransferQueueFamilyIndex == -1)
 	{
-		// If we didn't find a dedicated Queue, use the default one
-		TransferQueueFamilyIndex = ComputeQueueFamilyIndex;
+		// If we didn't find a dedicated Queue, leave it null
+		Queues[(int32)EVulkanQueueType::Transfer] = nullptr;
 	}
 	else
 	{
+		Queues[(int32)EVulkanQueueType::Transfer] = new FVulkanQueue(*this, TransferQueueFamilyIndex, EVulkanQueueType::Transfer);
 		ActiveQueueFamilies.Add(TransferQueueFamilyIndex);
 	}
-	Queues[(int32)EVulkanQueueType::Transfer] = new FVulkanQueue(*this, TransferQueueFamilyIndex, EVulkanQueueType::Transfer);
-
-	uint64 NumBits = QueueFamilyProps[GfxQueueFamilyIndex].timestampValidBits;
-	if (NumBits > 0)
-	{
-		ensure(NumBits == QueueFamilyProps[ComputeQueueFamilyIndex].timestampValidBits);
-		if (NumBits == 64)
-		{
-			// Undefined behavior trying to 1 << 64 on uint64
-			TimestampValidBitsMask = UINT64_MAX;
-		}
-		else
-		{
-			TimestampValidBitsMask = ((uint64)1 << (uint64)NumBits) - (uint64)1;
-		}
-	}
 
 	// Enumerate the available shading rates
 	if (OptionalDeviceExtensions.HasKHRFragmentShadingRate)
@@ -1694,9 +1695,18 @@ FGPUTimingCalibrationTimestamp FVulkanDevice::GetCalibrationTimestamp()
 
 void FVulkanDevice::ForEachQueue(TFunctionRef<void(FVulkanQueue&)> Callback)
 {
-	for (int32 QueueIndex = 0; QueueIndex < (int32)EVulkanQueueType::Count; ++QueueIndex)
+	static_assert((int32)EVulkanQueueType::Count == 3u, "New entries in EVulkanQueueType must be added in FVulkanDevice::ForEachQueue.");
+
+	Callback(*GetGraphicsQueue());
+
+	if (HasAsyncComputeQueue())
+	{
+		Callback(*GetComputeQueue());
+	}
+
+	if (HasTransferQueue())
 	{
-		Callback(*Queues[QueueIndex]);
+		Callback(*GetTransferQueue());
 	}
 }
 
@@ -1743,7 +1753,7 @@ VkBuffer FVulkanDevice::CreateBuffer(VkDeviceSize BufferSize, VkBufferUsageFlags
 		BufferCreateInfo.usage |= VK_BUFFER_USAGE_SHADER_DEVICE_ADDRESS_BIT;
 	}
 
-	if (bAsyncComputeQueue && GVulkanAllowConcurrentBuffer)
+	if (HasMultipleQueues() && GVulkanAllowConcurrentBuffer)
 	{
 		BufferCreateInfo.sharingMode = VK_SHARING_MODE_CONCURRENT;
 		BufferCreateInfo.queueFamilyIndexCount = ActiveQueueFamilies.Num();
diff --git a/Engine/Source/Runtime/VulkanRHI/Private/VulkanDevice.h b/Engine/Source/Runtime/VulkanRHI/Private/VulkanDevice.h
--- a/Engine/Source/Runtime/VulkanRHI/Private/VulkanDevice.h
+++ b/Engine/Source/Runtime/VulkanRHI/Private/VulkanDevice.h
@@ -280,7 +280,17 @@ public:
 
 	bool HasAsyncComputeQueue() const
 	{
-		return bAsyncComputeQueue;
+		return Queues[(int32)EVulkanQueueType::AsyncCompute] != nullptr;
+	}
+
+	bool HasTransferQueue() const
+	{
+		return Queues[(int32)EVulkanQueueType::Transfer] != nullptr;
+	}
+
+	bool HasMultipleQueues() const
+	{
+		return HasAsyncComputeQueue() || HasTransferQueue();
 	}
 
 	bool CanPresentOnComputeQueue() const
@@ -391,11 +401,6 @@ public:
 
 	bool SupportsBindless() const;
 
-	uint64 GetTimestampValidBitsMask() const
-	{
-		return TimestampValidBitsMask;
-	}
-
 	const VkComponentMapping& GetFormatComponentMapping(EPixelFormat UEFormat) const;
 
 	VkDevice GetInstanceHandle() const
@@ -646,11 +651,8 @@ private:
 	FCriticalSection BarrierEventLock;
 	TArray<VkEvent> BarrierEvents;
 
-	uint64 TimestampValidBitsMask = 0;
-
 	TStaticArray<FVulkanQueue*, (int32)EVulkanQueueType::Count> Queues;
 	FVulkanQueue* PresentQueue = nullptr;  // points to an existing queue
-	bool bAsyncComputeQueue = false;
 	bool bPresentOnComputeQueue = false;
 	TArray<uint32> ActiveQueueFamilies;
 
diff --git a/Engine/Source/Runtime/VulkanRHI/Private/VulkanRHI.cpp b/Engine/Source/Runtime/VulkanRHI/Private/VulkanRHI.cpp
--- a/Engine/Source/Runtime/VulkanRHI/Private/VulkanRHI.cpp
+++ b/Engine/Source/Runtime/VulkanRHI/Private/VulkanRHI.cpp
@@ -1691,7 +1691,9 @@ void* FVulkanDynamicRHI::RHIGetNativeGraphicsQueue()
 
 void* FVulkanDynamicRHI::RHIGetNativeComputeQueue()
 {
-	return (void*)Device->GetComputeQueue()->GetHandle();
+	return Device->HasAsyncComputeQueue() ? 
+		(void*)Device->GetComputeQueue()->GetHandle() :
+		(void*)Device->GetGraphicsQueue()->GetHandle();
 }
 
 void* FVulkanDynamicRHI::RHIGetNativeInstance()
diff --git a/Engine/Source/Runtime/VulkanRHI/Private/VulkanSwapChain.cpp b/Engine/Source/Runtime/VulkanRHI/Private/VulkanSwapChain.cpp
--- a/Engine/Source/Runtime/VulkanRHI/Private/VulkanSwapChain.cpp
+++ b/Engine/Source/Runtime/VulkanRHI/Private/VulkanSwapChain.cpp
@@ -895,27 +895,23 @@ void FVulkanDevice::SetupPresentQueue(VkSurfaceKHR Surface)
 			return (bSupportsPresent == VK_TRUE);
 		};
 
-		bool bGfx = SupportsPresent(Gpu, GetGraphicsQueue());
+		const bool bGfx = SupportsPresent(Gpu, GetGraphicsQueue());
 		if (!bGfx)
 		{
 			FPlatformMisc::MessageBoxExt(EAppMsgType::Ok, TEXT("Cannot find a compatible Vulkan device that supports surface presentation.\n\n"), TEXT("Vulkan device not available"));
 			FPlatformMisc::RequestExitWithStatus(true, 1);
 		}
+		PresentQueue = GetGraphicsQueue();
 
-		bool bCompute = SupportsPresent(Gpu, GetComputeQueue());
-		if (GetTransferQueue()->GetFamilyIndex() != GetGraphicsQueue()->GetFamilyIndex() && GetTransferQueue()->GetFamilyIndex() != GetComputeQueue()->GetFamilyIndex())
+		if (HasAsyncComputeQueue())
 		{
-			SupportsPresent(Gpu, GetTransferQueue());
-		}
-		if (GAllowPresentOnComputeQueue.GetValueOnAnyThread() != 0 && GetComputeQueue()->GetFamilyIndex() != GetGraphicsQueue()->GetFamilyIndex() && bCompute)
-		{
-			//#todo-rco: Do other IHVs have a fast path here?
-			bPresentOnComputeQueue = (VendorId == EGpuVendorId::Amd);
-			PresentQueue = GetComputeQueue();
-		}
-		else
-		{
-			PresentQueue = GetGraphicsQueue();
+			const bool bCompute = SupportsPresent(Gpu, GetComputeQueue());
+			if (bCompute && (GAllowPresentOnComputeQueue.GetValueOnAnyThread() != 0))
+			{
+				//#todo-rco: Do other IHVs have a fast path here?
+				bPresentOnComputeQueue = (VendorId == EGpuVendorId::Amd);
+				PresentQueue = GetComputeQueue();
+			}
 		}
 	}
 }
diff --git a/Engine/Source/Runtime/VulkanRHI/Private/VulkanTexture.cpp b/Engine/Source/Runtime/VulkanRHI/Private/VulkanTexture.cpp
--- a/Engine/Source/Runtime/VulkanRHI/Private/VulkanTexture.cpp
+++ b/Engine/Source/Runtime/VulkanRHI/Private/VulkanTexture.cpp
@@ -502,7 +502,7 @@ void FVulkanTexture::GenerateImageCreateInfo(
 		}
 	}
 
-	if (InDevice.HasAsyncComputeQueue() && (GVulkanAllowConcurrentImage != 0))
+	if (InDevice.HasMultipleQueues() && (GVulkanAllowConcurrentImage != 0))
 	{
 		ImageCreateInfo.sharingMode = VK_SHARING_MODE_CONCURRENT;
 		ImageCreateInfo.queueFamilyIndexCount = InDevice.GetActiveQueueFamilies().Num();
