VulkanRHI: Move DescriptorPoolSetContainer to Context.

#rnx

diff --git a/Engine/Source/Runtime/VulkanRHI/Private/VulkanCommandBuffer.cpp b/Engine/Source/Runtime/VulkanRHI/Private/VulkanCommandBuffer.cpp
--- a/Engine/Source/Runtime/VulkanRHI/Private/VulkanCommandBuffer.cpp
+++ b/Engine/Source/Runtime/VulkanRHI/Private/VulkanCommandBuffer.cpp
@@ -108,12 +108,6 @@ void FVulkanCommandBuffer::BeginRenderPass(const FVulkanBeginRenderPassInfo& Beg
 	}
 
 	State = EState::IsInsideRenderPass;
-
-	// Acquire a descriptor pool set on a first render pass
-	if (CurrentDescriptorPoolSetContainer == nullptr)
-	{
-		AcquirePoolSetContainer();
-	}
 }
 
 void FVulkanCommandBuffer::End(FVulkanQueryPool* QueryPool)
@@ -179,8 +173,6 @@ void FVulkanCommandBuffer::Begin(FVulkanQueryPool* QueryPool, VkRenderPass Rende
 
 	VERIFYVULKANRESULT(VulkanRHI::vkBeginCommandBuffer(CommandBufferHandle, &BeginInfo));
 
-	check(!CurrentDescriptorPoolSetContainer);
-
 	if (Device.SupportsBindless())
 	{
 		FVulkanBindlessDescriptorManager* BindlessDescriptorManager = Device.GetBindlessDescriptorManager();
@@ -208,38 +200,6 @@ void FVulkanCommandBuffer::Begin(FVulkanQueryPool* QueryPool, VkRenderPass Rende
 #endif
 }
 
-void FVulkanCommandBuffer::AcquirePoolSetContainer()
-{
-	check(!CurrentDescriptorPoolSetContainer);
-	CurrentDescriptorPoolSetContainer = &Device.GetDescriptorPoolsManager().AcquirePoolSetContainer();
-	ensure(TypedDescriptorPoolSets.Num() == 0);
-}
-
-bool FVulkanCommandBuffer::AcquirePoolSetAndDescriptorsIfNeeded(const class FVulkanDescriptorSetsLayout& Layout, bool bNeedDescriptors, VkDescriptorSet* OutDescriptors)
-{
-	//#todo-rco: This only happens when we call draws outside a render pass...
-	if (!CurrentDescriptorPoolSetContainer)
-	{
-		AcquirePoolSetContainer();
-	}
-
-	const uint32 Hash = VULKAN_HASH_POOLS_WITH_TYPES_USAGE_ID ? Layout.GetTypesUsageID() : GetTypeHash(Layout);
-	FVulkanTypedDescriptorPoolSet*& FoundTypedSet = TypedDescriptorPoolSets.FindOrAdd(Hash);
-
-	if (!FoundTypedSet)
-	{
-		FoundTypedSet = CurrentDescriptorPoolSetContainer->AcquireTypedPoolSet(Layout);
-		bNeedDescriptors = true;
-	}
-
-	if (bNeedDescriptors)
-	{
-		return FoundTypedSet->AllocateDescriptorSets(Layout, OutDescriptors);
-	}
-
-	return false;
-}
-
 void FVulkanCommandBuffer::Reset()
 {
 	// Reset the secondary command buffers we executed from this one
@@ -262,18 +222,6 @@ void FVulkanCommandBuffer::Reset()
 		CurrentScissors.Empty();
 		CurrentStencilRef = 0;
 
-		if (CurrentDescriptorPoolSetContainer)
-		{
-			//#todo-rco: Reset here?
-			TypedDescriptorPoolSets.Reset();
-			Device.GetDescriptorPoolsManager().ReleasePoolSet(*CurrentDescriptorPoolSetContainer);
-			CurrentDescriptorPoolSetContainer = nullptr;
-		}
-		else
-		{
-			check(TypedDescriptorPoolSets.Num() == 0);
-		}
-
 		for (VkEvent BarrierEvent : EndedBarrierEvents)
 		{
 			Device.ReleaseBarrierEvent(BarrierEvent);
diff --git a/Engine/Source/Runtime/VulkanRHI/Private/VulkanCommandBuffer.h b/Engine/Source/Runtime/VulkanRHI/Private/VulkanCommandBuffer.h
--- a/Engine/Source/Runtime/VulkanRHI/Private/VulkanCommandBuffer.h
+++ b/Engine/Source/Runtime/VulkanRHI/Private/VulkanCommandBuffer.h
@@ -116,11 +116,6 @@ public:
 	void BeginSplitBarrier(VkEvent BarrierEvent, const VkDependencyInfo& DependencyInfo);
 	void EndSplitBarrier(VkEvent BarrierEvent, const VkDependencyInfo& DependencyInfo);
 
-	//#todo-rco: Hide this
-	FVulkanDescriptorPoolSetContainer* CurrentDescriptorPoolSetContainer = nullptr;
-
-	bool AcquirePoolSetAndDescriptorsIfNeeded(const class FVulkanDescriptorSetsLayout& Layout, bool bNeedDescriptors, VkDescriptorSet* OutDescriptors);
-
 #if RHI_NEW_GPU_PROFILER
 	template <typename TEventType, typename... TArgs>
 	TEventType& EmplaceProfilerEvent(TArgs&&... Args)
@@ -171,8 +166,6 @@ private:
 
 	TArray<VkEvent> EndedBarrierEvents;
 
-	void AcquirePoolSetContainer();
-
 	void AllocMemory();
 	void FreeMemory();
 
@@ -184,10 +177,6 @@ private:
 	uint64 EndTimestamp = 0;
 #endif
 
-public:
-	//#todo-rco: Hide this
-	TMap<uint32, class FVulkanTypedDescriptorPoolSet*> TypedDescriptorPoolSets;
-
 	friend class FVulkanDynamicRHI;
 	friend FVulkanContextCommon;
 };
diff --git a/Engine/Source/Runtime/VulkanRHI/Private/VulkanContext.cpp b/Engine/Source/Runtime/VulkanRHI/Private/VulkanContext.cpp
--- a/Engine/Source/Runtime/VulkanRHI/Private/VulkanContext.cpp
+++ b/Engine/Source/Runtime/VulkanRHI/Private/VulkanContext.cpp
@@ -63,8 +63,6 @@ void FVulkanContextCommon::EndPayload()
 // payloads to the given array. Resets the context so new commands can be recorded.
 void FVulkanContextCommon::Finalize(TArray<FVulkanPayload*>& OutPayloads)
 {
-	FlushProfilerStats();
-
 	FlushPendingSyncPoints();
 
 	if (ContextSyncPoint.IsValid())
@@ -314,6 +312,56 @@ void FVulkanCommandListContext::RHIEndParallelRenderPass()
 	CurrentParallelRenderPassInfo = nullptr;
 }
 
+void FVulkanCommandListContext::Finalize(TArray<FVulkanPayload*>& OutPayloads)
+{
+	FlushProfilerStats();
+
+	if (CurrentDescriptorPoolSetContainer)
+	{
+		GetPayload(EPhase::Signal).DescriptorPoolSetContainers.Add(CurrentDescriptorPoolSetContainer);
+		CurrentDescriptorPoolSetContainer = nullptr;
+		TypedDescriptorPoolSets.Reset();
+	}
+	else
+	{
+		check(TypedDescriptorPoolSets.Num() == 0);
+	}
+
+	FVulkanContextCommon::Finalize(OutPayloads);
+}
+
+
+
+void FVulkanCommandListContext::AcquirePoolSetContainer()
+{
+	if (!CurrentDescriptorPoolSetContainer)
+	{
+		CurrentDescriptorPoolSetContainer = &Device.GetDescriptorPoolsManager().AcquirePoolSetContainer();
+		ensure(TypedDescriptorPoolSets.Num() == 0);
+	}
+}
+
+bool FVulkanCommandListContext::AcquirePoolSetAndDescriptorsIfNeeded(const FVulkanDescriptorSetsLayout& Layout, bool bNeedDescriptors, VkDescriptorSet* OutDescriptors)
+{
+	AcquirePoolSetContainer();
+
+	const uint32 Hash = VULKAN_HASH_POOLS_WITH_TYPES_USAGE_ID ? Layout.GetTypesUsageID() : GetTypeHash(Layout);
+	FVulkanTypedDescriptorPoolSet*& FoundTypedSet = TypedDescriptorPoolSets.FindOrAdd(Hash);
+
+	if (!FoundTypedSet)
+	{
+		FoundTypedSet = CurrentDescriptorPoolSetContainer->AcquireTypedPoolSet(Layout);
+		bNeedDescriptors = true;
+	}
+
+	if (bNeedDescriptors)
+	{
+		return FoundTypedSet->AllocateDescriptorSets(Layout, OutDescriptors);
+	}
+
+	return false;
+}
+
 TLockFreePointerListUnordered<FVulkanUploadContext, PLATFORM_CACHE_LINE_SIZE> FVulkanUploadContext::Pool;
 void FVulkanUploadContext::DestroyPool()
 {
diff --git a/Engine/Source/Runtime/VulkanRHI/Private/VulkanContext.h b/Engine/Source/Runtime/VulkanRHI/Private/VulkanContext.h
--- a/Engine/Source/Runtime/VulkanRHI/Private/VulkanContext.h
+++ b/Engine/Source/Runtime/VulkanRHI/Private/VulkanContext.h
@@ -196,9 +196,8 @@ public:
 		return nullptr;
 	}
 
-	virtual void FlushProfilerStats() { /* do nothing */ }
-
 private:
+
 	FVulkanQueryPool* GetCurrentTimestampQueryPool(FVulkanPayload& Payload);
 	void PrepareNewCommandBuffer(FVulkanPayload& Payload);
 	void NewPayload();
@@ -304,7 +303,7 @@ public:
 	virtual void RHIEndRenderQuery(FRHIRenderQuery* RenderQuery) final override;
 
 #if RHI_NEW_GPU_PROFILER
-	virtual void FlushProfilerStats() override
+	void FlushProfilerStats()
 	{
 		// Flush accumulated draw stats (if breadcrumbs are available to attach them to)
 		if (StatEvent && bSupportsBreadcrumbs)
@@ -317,6 +316,8 @@ public:
 	virtual void RHICalibrateTimers(FRHITimestampCalibrationQuery* CalibrationQuery) final override;
 #endif
 
+	virtual void Finalize(TArray<FVulkanPayload*>& OutPayloads) override;
+
 	virtual void RHIBeginDrawingViewport(FRHIViewport* Viewport, FRHITexture* RenderTargetRHI) final override;
 	virtual void RHIEndDrawingViewport(FRHIViewport* Viewport, bool bPresent, bool bLockToVsync) final override;
 
@@ -399,6 +400,8 @@ public:
 		CurrentParallelRenderPassInfo = ParallelRenderPassInfo;
 	}
 
+	bool AcquirePoolSetAndDescriptorsIfNeeded(const class FVulkanDescriptorSetsLayout& Layout, bool bNeedDescriptors, VkDescriptorSet* OutDescriptors);
+
 protected:
 	FVulkanCommandListContext* Immediate;
 	const ERHIPipeline RHIPipeline;
@@ -427,6 +430,10 @@ protected:
 
 	void InternalClearMRT(FVulkanCommandBuffer& CommandBuffer, bool bClearColor, int32 NumClearColors, const FLinearColor* ColorArray, bool bClearDepth, float Depth, bool bClearStencil, uint32 Stencil);
 
+	FVulkanDescriptorPoolSetContainer* CurrentDescriptorPoolSetContainer = nullptr;
+	TMap<uint32, class FVulkanTypedDescriptorPoolSet*> TypedDescriptorPoolSets;
+	void AcquirePoolSetContainer();
+
 public:
 	bool IsSwapchainImage(FRHITexture* InTexture) const;
 	VkFormat GetSwapchainImageFormat() const;
diff --git a/Engine/Source/Runtime/VulkanRHI/Private/VulkanPipelineState.cpp b/Engine/Source/Runtime/VulkanRHI/Private/VulkanPipelineState.cpp
--- a/Engine/Source/Runtime/VulkanRHI/Private/VulkanPipelineState.cpp
+++ b/Engine/Source/Runtime/VulkanRHI/Private/VulkanPipelineState.cpp
@@ -203,7 +203,7 @@ bool FVulkanComputePipelineDescriptorState::InternalUpdateDescriptorSets(FVulkan
 	// We are not using UseVulkanDescriptorCache() for compute pipelines
 	// Compute tend to use volatile resources that polute descriptor cache
 
-	if (!Context.GetCommandBuffer().AcquirePoolSetAndDescriptorsIfNeeded(*DescriptorSetsLayout, true, DescriptorSetHandles.GetData()))
+	if (!Context.AcquirePoolSetAndDescriptorsIfNeeded(*DescriptorSetsLayout, true, DescriptorSetHandles.GetData()))
 	{
 		return false;
 	}
@@ -427,7 +427,7 @@ bool FVulkanGraphicsPipelineDescriptorState::InternalUpdateDescriptorSets(FVulka
 		const bool bNeedsWrite = (bIsResourcesDirty || ShouldAlwaysWriteDescriptors());
 
 		// Allocate sets based on what changed
-		if (Context.GetCommandBuffer().AcquirePoolSetAndDescriptorsIfNeeded(*DescriptorSetsLayout, bNeedsWrite, DescriptorSetHandles.GetData()))
+		if (Context.AcquirePoolSetAndDescriptorsIfNeeded(*DescriptorSetsLayout, bNeedsWrite, DescriptorSetHandles.GetData()))
 		{
 			uint32 RemainingSetsMask = UsedSetsMask;
 			uint32 Set = 0;
@@ -604,11 +604,7 @@ void FVulkanCommandListContext::RHISetGraphicsPipelineState(FRHIGraphicsPipeline
 
 void FVulkanCommandListContext::RHISetComputePipelineState(FRHIComputePipelineState* ComputePipelineState)
 {
-	FVulkanCommandBuffer* CmdBuffer = GetActiveCmdBuffer();
-	if (CmdBuffer->CurrentDescriptorPoolSetContainer == nullptr)
-	{
-		CmdBuffer->CurrentDescriptorPoolSetContainer = &Device.GetDescriptorPoolsManager().AcquirePoolSetContainer();
-	}
+	AcquirePoolSetContainer();
 
 	//#todo-rco: Set PendingGfx to null
 	FVulkanComputePipeline* ComputePipeline = ResourceCast(ComputePipelineState);
diff --git a/Engine/Source/Runtime/VulkanRHI/Private/VulkanRenderTarget.cpp b/Engine/Source/Runtime/VulkanRHI/Private/VulkanRenderTarget.cpp
--- a/Engine/Source/Runtime/VulkanRHI/Private/VulkanRenderTarget.cpp
+++ b/Engine/Source/Runtime/VulkanRHI/Private/VulkanRenderTarget.cpp
@@ -518,6 +518,9 @@ void FVulkanCommandListContext::RHIBeginRenderPass(const FRHIRenderPassInfo& InI
 		BeginOcclusionQueryBatch(InInfo.NumOcclusionQueries);
 	}
 
+	// Acquire a descriptor pool set on a first render pass
+	AcquirePoolSetContainer();
+
 	const bool bNeedsAllPlanes = Device.NeedsAllPlanes();
 
 	FRHITexture* DSTexture = InInfo.DepthStencilRenderTarget.DepthStencilTarget;
diff --git a/Engine/Source/Runtime/VulkanRHI/Private/VulkanSubmission.cpp b/Engine/Source/Runtime/VulkanRHI/Private/VulkanSubmission.cpp
--- a/Engine/Source/Runtime/VulkanRHI/Private/VulkanSubmission.cpp
+++ b/Engine/Source/Runtime/VulkanRHI/Private/VulkanSubmission.cpp
@@ -10,6 +10,7 @@
 #include "VulkanDevice.h"
 #include "VulkanSubmission.h"
 #include "VulkanQueue.h"
+#include "VulkanDescriptorSets.h"
 
 #include "HAL/Runnable.h"
 #include "HAL/RunnableThread.h"
@@ -568,6 +569,11 @@ void FVulkanDynamicRHI::CompletePayload(FVulkanPayload* Payload)
 		CommandBuffer->Reset();
 	}
 
+	for (FVulkanDescriptorPoolSetContainer* DescriptorPoolSetContainer : Payload->DescriptorPoolSetContainers)
+	{
+		Device->GetDescriptorPoolsManager().ReleasePoolSet(*DescriptorPoolSetContainer);
+	}
+
 	struct FQueryResult
 	{
 		uint64 Result;
diff --git a/Engine/Source/Runtime/VulkanRHI/Private/VulkanSubmission.h b/Engine/Source/Runtime/VulkanRHI/Private/VulkanSubmission.h
--- a/Engine/Source/Runtime/VulkanRHI/Private/VulkanSubmission.h
+++ b/Engine/Source/Runtime/VulkanRHI/Private/VulkanSubmission.h
@@ -16,6 +16,7 @@ class FVulkanPendingComputeState;
 class FVulkanQueue;
 class FVulkanSwapChain;
 class FVulkanTiming;
+class FVulkanDescriptorPoolSetContainer;
 
 // NOTE: typedef to FGraphEvent for now, but eventually move all syncs inside a FVulkanSyncPoint class
 using FVulkanSyncPoint = FGraphEvent;
@@ -66,6 +67,9 @@ protected:
 	// Queries used in the command lists of this payload
 	TStaticArray<TArray<FVulkanQueryPool*>, (int32)EVulkanQueryPoolType::Count> QueryPools;
 
+	// Used DescriptorPoolSetContainers that can be cleaned up on completion
+	TArray<FVulkanDescriptorPoolSetContainer*> DescriptorPoolSetContainers;
+
 	// Used by RHIRunOnQueue
 	TFunction<void(VkQueue)> PreExecuteCallback;
 
