VulkanRHI: Create fallback for reference timestamps when GetCalibration values don't line up with the last read timestamps.

#rnx

diff --git a/Engine/Source/Runtime/VulkanRHI/Private/VulkanDevice.cpp b/Engine/Source/Runtime/VulkanRHI/Private/VulkanDevice.cpp
--- a/Engine/Source/Runtime/VulkanRHI/Private/VulkanDevice.cpp
+++ b/Engine/Source/Runtime/VulkanRHI/Private/VulkanDevice.cpp
@@ -1661,7 +1661,7 @@ void FVulkanDevice::GetCalibrationTimestamp(FVulkanTiming& InOutTiming)
 	else
 	{
 		InOutTiming.GPUTimestamp = 0;
-		InOutTiming.CPUTimestamp = 0;
+		InOutTiming.CPUTimestamp = FPlatformTime::Cycles64();
 	}
 }
 #else
diff --git a/Engine/Source/Runtime/VulkanRHI/Private/VulkanSubmission.cpp b/Engine/Source/Runtime/VulkanRHI/Private/VulkanSubmission.cpp
--- a/Engine/Source/Runtime/VulkanRHI/Private/VulkanSubmission.cpp
+++ b/Engine/Source/Runtime/VulkanRHI/Private/VulkanSubmission.cpp
@@ -617,12 +617,40 @@ void FVulkanDynamicRHI::CompletePayload(FVulkanPayload* Payload)
 							// Convert from GPU timestamp to CPU timestamp (relative to FPlatformTime::Cycles64())
 							checkSlow(Payload->Queue.Timing);
 							const FVulkanTiming* CurrentTiming = Payload->Queue.Timing;
-							checkf(TempResults[QueryIndex].Result >= CurrentTiming->GPUTimestamp,
-								TEXT("Invalid timestamp can't be moved to CPU timestamp (%lld < %lld)"),
-								TempResults[QueryIndex].Result, CurrentTiming->GPUTimestamp);
-							const uint64 GPUDelta = TempResults[QueryIndex].Result - CurrentTiming->GPUTimestamp;
+
+							uint64 CPUCalibrationTimestamp = CurrentTiming->CPUTimestamp;
+							uint64 GPUCalibrationTimestamp = CurrentTiming->GPUTimestamp;
+
+							// NOTE: Adreno GPU calibration values aren't valid, create a fallback
+							static bool bBypassTiming = false;
+							if (bBypassTiming || (TempResults[QueryIndex].Result < GPUCalibrationTimestamp) || (GPUCalibrationTimestamp == 0))
+							{
+								static uint64 ReferenceCPUTimestamp = 0;
+								static uint64 ReferenceGPUTimestamp = 0;
+
+								if (!bBypassTiming)
+								{
+									UE_LOG(LogVulkanRHI, Warning, 
+										TEXT("GPU timestamp can't be moved to CPU timeline.  ")
+										TEXT("Calibration data is invalid (%lld < %lld)!  ")
+										TEXT("CPU and GPU profiling data will not line up!"),
+										TempResults[QueryIndex].Result, GPUCalibrationTimestamp);
+
+									const uint64 CurrentTimestamp = FPlatformTime::Cycles64();
+									ReferenceCPUTimestamp = ((CPUCalibrationTimestamp != 0) && 
+										(CPUCalibrationTimestamp < CurrentTimestamp)) ?
+										CPUCalibrationTimestamp : CurrentTimestamp;
+									ReferenceGPUTimestamp = TempResults[QueryIndex].Result;
+									bBypassTiming = true;
+								}
+
+								CPUCalibrationTimestamp = ReferenceCPUTimestamp;
+								GPUCalibrationTimestamp = ReferenceGPUTimestamp;
+							}
+
+							const uint64 GPUDelta = TempResults[QueryIndex].Result - GPUCalibrationTimestamp;
 							const uint64 CPUDelta = (GPUDelta * CurrentTiming->CPUFrequency) / CurrentTiming->GPUFrequency;
-							(*ResultPtr) = CPUDelta + CurrentTiming->CPUTimestamp;
+							(*ResultPtr) = CPUCalibrationTimestamp + CPUDelta;
 						}
 						else
 #endif
