VulkanRHI: When only SRVGraphicsNonPixel is present in transition, do not remove the VS bit.

#rnx

diff --git a/Engine/Source/Runtime/VulkanRHI/Private/VulkanBarriers.cpp b/Engine/Source/Runtime/VulkanRHI/Private/VulkanBarriers.cpp
--- a/Engine/Source/Runtime/VulkanRHI/Private/VulkanBarriers.cpp
+++ b/Engine/Source/Runtime/VulkanRHI/Private/VulkanBarriers.cpp
@@ -869,7 +869,9 @@ void ConvertTransitionToBarriers(FVulkanCommandListContext& Context, const FVulk
 			GetVkStageAndAccessFlags(ResourceState.AccessAfter, UnderlyingType, UsageFlags, bIsDepthStencil, bSupportsReadOnlyOptimal, DstStageMask, DstAccessFlags, DstLayout, false);
 
 			// If not compute, remove vertex pipeline bits as only compute updates vertex buffers
-			if (!(SrcStageMask & VK_PIPELINE_STAGE_COMPUTE_SHADER_BIT))
+			// EXCEPT the special case for SRVGraphicsNonPixel where Vertex needs to be preserved
+			const bool bKeepVertex = Texture && EnumHasAnyFlags(ResourceState.AccessAfter, ERHIAccess::SRVGraphicsNonPixel) && !EnumHasAnyFlags(ResourceState.AccessAfter, ERHIAccess::SRVGraphicsPixel);
+			if (!bKeepVertex && !VKHasAllFlags(SrcStageMask, VK_PIPELINE_STAGE_COMPUTE_SHADER_BIT))
 			{
 				DstStageMask &= ~(VK_PIPELINE_STAGE_VERTEX_SHADER_BIT | VK_PIPELINE_STAGE_GEOMETRY_SHADER_BIT);
 			}
