VulkanRHI:  reset the calibration fallback every frame to limit the drift over time (even if the profiler tracks remain unaligned).

#rnx

diff --git a/Engine/Source/Runtime/VulkanRHI/Private/VulkanSubmission.cpp b/Engine/Source/Runtime/VulkanRHI/Private/VulkanSubmission.cpp
--- a/Engine/Source/Runtime/VulkanRHI/Private/VulkanSubmission.cpp
+++ b/Engine/Source/Runtime/VulkanRHI/Private/VulkanSubmission.cpp
@@ -579,7 +579,14 @@ void FVulkanDynamicRHI::CompletePayload(FVulkanPayload* Payload)
 		uint64 Result;
 		uint64 Availability;
 	};
+
 	static TArray<FQueryResult> TempResults;
+
+	// Fallback for devices without calibrated_timestamps (or with unexpected values from calibration)
+	static bool bBypassTiming = false;
+	static uint64 ReferenceCPUTimestamp = 0;
+	static uint64 ReferenceGPUTimestamp = 0;
+
 	for (TArray<FVulkanQueryPool*>& QueryPoolArray : Payload->QueryPools)
 	{
 		for (FVulkanQueryPool* QueryPool : QueryPoolArray)
@@ -622,19 +629,20 @@ void FVulkanDynamicRHI::CompletePayload(FVulkanPayload* Payload)
 							uint64 GPUCalibrationTimestamp = CurrentTiming->GPUTimestamp;
 
 							// NOTE: Adreno GPU calibration values aren't valid, create a fallback
-							static bool bBypassTiming = false;
 							if (bBypassTiming || (TempResults[QueryIndex].Result < GPUCalibrationTimestamp) || (GPUCalibrationTimestamp == 0))
 							{
-								static uint64 ReferenceCPUTimestamp = 0;
-								static uint64 ReferenceGPUTimestamp = 0;
-
 								if (!bBypassTiming)
 								{
-									UE_LOG(LogVulkanRHI, Warning, 
-										TEXT("GPU timestamp can't be moved to CPU timeline.  ")
-										TEXT("Calibration data is invalid (%lld < %lld)!  ")
-										TEXT("CPU and GPU profiling data will not line up!"),
-										TempResults[QueryIndex].Result, GPUCalibrationTimestamp);
+									static bool bWarnOnce = false;
+									if (!bWarnOnce)
+									{
+										bWarnOnce = true;
+										UE_LOG(LogVulkanRHI, Warning,
+											TEXT("GPU timestamp can't be moved to CPU timeline.  ")
+											TEXT("Calibration data is invalid (Curr=%lld  Cal=%lld)!  ")
+											TEXT("CPU and GPU profiling data will not line up!"),
+											TempResults[QueryIndex].Result, GPUCalibrationTimestamp);
+									}
 
 									const uint64 CurrentTimestamp = FPlatformTime::Cycles64();
 									ReferenceCPUTimestamp = ((CPUCalibrationTimestamp != 0) && 
@@ -682,6 +690,11 @@ void FVulkanDynamicRHI::CompletePayload(FVulkanPayload* Payload)
 	{
 		// Switch the new timing struct into the queue. This redirects timestamp results to separate each frame's work.
 		Payload->Queue.Timing = Payload->Timing.GetValue();
+
+		// Reset the fallback to recalibrated it every frame (it remains innacurate, but at least we limit drift)
+		bBypassTiming = false;
+		ReferenceCPUTimestamp = 0;
+		ReferenceGPUTimestamp = 0;
 	}
 #else
 	if (GRHIGlobals.SupportsTimestampRenderQueries && !FVulkanPlatform::HasCustomFrameTiming())
