VulkanRHI: Remove unused crash handler (was part of old profiler).  Add support for VK_TYPE_TO_STRING on Android Debug/Dev builds.

#rnx

diff --git a/Engine/Source/Runtime/VulkanRHI/Private/Android/VulkanAndroidPlatform.cpp b/Engine/Source/Runtime/VulkanRHI/Private/Android/VulkanAndroidPlatform.cpp
--- a/Engine/Source/Runtime/VulkanRHI/Private/Android/VulkanAndroidPlatform.cpp
+++ b/Engine/Source/Runtime/VulkanRHI/Private/Android/VulkanAndroidPlatform.cpp
@@ -2136,10 +2136,3 @@ FString FVulkanAndroidPlatform::GetVulkanProfileNameForFeatureLevel(ERHIFeatureL
 	}
 	return ProfileName;
 } 
-
-void FVulkanAndroidPlatform::WriteCrashMarker(const FOptionalVulkanDeviceExtensions& OptionalExtensions, FVulkanCommandBuffer* CmdBuffer, VkBuffer DestBuffer, const TArrayView<uint32>& Entries, bool bAdding)
-{
-	ensure(Entries.Num() <= GMaxCrashBufferEntries);
-
-	WriteCrashMarkerWithoutExtensions(CmdBuffer, DestBuffer, Entries, bAdding);
-}
diff --git a/Engine/Source/Runtime/VulkanRHI/Private/Android/VulkanAndroidPlatform.h b/Engine/Source/Runtime/VulkanRHI/Private/Android/VulkanAndroidPlatform.h
--- a/Engine/Source/Runtime/VulkanRHI/Private/Android/VulkanAndroidPlatform.h
+++ b/Engine/Source/Runtime/VulkanRHI/Private/Android/VulkanAndroidPlatform.h
@@ -19,16 +19,20 @@
 #define VULKAN_PURGE_SHADER_MODULES					0
 #define VULKAN_SUPPORTS_DEDICATED_ALLOCATION		0
 #define VULKAN_SUPPORTS_ASTC_DECODE_MODE			1
-#define VULKAN_SUPPORTS_NV_DIAGNOSTIC_CHECKPOINT	0
 #define VULKAN_SUPPORTS_SCALAR_BLOCK_LAYOUT			1
 #define VULKAN_SUPPORTS_TRANSIENT_RESOURCE_ALLOCATOR 0
 #define VULKAN_SUPPORTS_DRIVER_PROPERTIES			0
 #define VULKAN_SUPPORTS_DESCRIPTOR_INDEXING			1
-#define VULKAN_SUPPORTS_GPU_CRASH_DUMPS				1
 #define VULKAN_SUPPORTS_RAY_TRACING_POSITION_FETCH	0
 
 #define UE_VK_API_VERSION							VK_API_VERSION_1_1
 
+#if UE_BUILD_DEBUG || UE_BUILD_DEVELOPMENT
+#	include "vk_enum_string_helper.h"
+#	define VK_TYPE_TO_STRING(Type, Value) ANSI_TO_TCHAR(string_##Type(Value))
+#	define VK_FLAGS_TO_STRING(Type, Value) ANSI_TO_TCHAR(string_##Type(Value).c_str())
+#endif
+
 #define ENUM_VK_ENTRYPOINTS_PLATFORM_BASE(EnumMacro)
 
 #define ENUM_VK_ENTRYPOINTS_PLATFORM_INSTANCE(EnumMacro) \
@@ -250,8 +254,6 @@ public:
 		return VK_SHADER_STAGE_FRAGMENT_BIT | VK_SHADER_STAGE_COMPUTE_BIT;
 	}
 
-	static void WriteCrashMarker(const FOptionalVulkanDeviceExtensions& OptionalExtensions, FVulkanCommandBuffer* CmdBuffer, VkBuffer DestBuffer, const TArrayView<uint32>& Entries, bool bAdding);
-
 	static VkTimeDomainKHR GetTimeDomain() { return VK_TIME_DOMAIN_CLOCK_MONOTONIC_KHR; }
 
 #if USE_ANDROID_VULKAN_SWAPPY
diff --git a/Engine/Source/Runtime/VulkanRHI/Private/Linux/VulkanLinuxPlatform.cpp b/Engine/Source/Runtime/VulkanRHI/Private/Linux/VulkanLinuxPlatform.cpp
--- a/Engine/Source/Runtime/VulkanRHI/Private/Linux/VulkanLinuxPlatform.cpp
+++ b/Engine/Source/Runtime/VulkanRHI/Private/Linux/VulkanLinuxPlatform.cpp
@@ -364,34 +364,3 @@ void FVulkanLinuxPlatform::CreateSurface(FVulkanPlatformWindowContext& WindowCon
 		check(0);
 	}
 }
-
-void FVulkanLinuxPlatform::WriteCrashMarker(const FOptionalVulkanDeviceExtensions& OptionalExtensions, FVulkanCommandBuffer* CmdBuffer, VkBuffer DestBuffer, const TArrayView<uint32>& Entries, bool bAdding)
-{
-	ensure(Entries.Num() <= GMaxCrashBufferEntries);
-
-	if (OptionalExtensions.HasAMDBufferMarker)
-	{
-		// AMD API only allows updating one entry at a time. Assume buffer has entry 0 as num entries
-		VulkanDynamicAPI::vkCmdWriteBufferMarkerAMD(CmdBuffer->GetHandle(), VK_PIPELINE_STAGE_BOTTOM_OF_PIPE_BIT, DestBuffer, 0, Entries.Num());
-		if (bAdding)
-		{
-			int32 LastIndex = Entries.Num() - 1;
-			// +1 size as entries start at index 1
-			VulkanDynamicAPI::vkCmdWriteBufferMarkerAMD(CmdBuffer->GetHandle(), VK_PIPELINE_STAGE_BOTTOM_OF_PIPE_BIT, DestBuffer, (1 + LastIndex) * sizeof(uint32), Entries[LastIndex]);
-		}
-	}
-	else
-	{
-		WriteCrashMarkerWithoutExtensions(CmdBuffer, DestBuffer, Entries, bAdding);
-	}
-
-	if (OptionalExtensions.HasNVDiagnosticCheckpoints)
-	{
-		if (bAdding)
-		{
-			int32 LastIndex = Entries.Num() - 1;
-			uint32 Value = Entries[LastIndex];
-			VulkanDynamicAPI::vkCmdSetCheckpointNV(CmdBuffer->GetHandle(), (void*)(size_t)Value);
-		}
-	}
-}
diff --git a/Engine/Source/Runtime/VulkanRHI/Private/Linux/VulkanLinuxPlatform.h b/Engine/Source/Runtime/VulkanRHI/Private/Linux/VulkanLinuxPlatform.h
--- a/Engine/Source/Runtime/VulkanRHI/Private/Linux/VulkanLinuxPlatform.h
+++ b/Engine/Source/Runtime/VulkanRHI/Private/Linux/VulkanLinuxPlatform.h
@@ -10,7 +10,6 @@
 #define VULKAN_SHOULD_DEBUG_IN_DEVELOPMENT			1
 #define VULKAN_SHOULD_ENABLE_DRAW_MARKERS			(UE_BUILD_DEBUG || UE_BUILD_DEVELOPMENT)
 #define VULKAN_SIGNAL_UNIMPLEMENTED()				checkf(false, TEXT("Unimplemented vulkan functionality: %hs"), __PRETTY_FUNCTION__)
-#define VULKAN_SUPPORTS_AMD_BUFFER_MARKER			1
 #define VULKAN_SUPPORTS_SCALAR_BLOCK_LAYOUT			1
 #define VULKAN_SUPPORTS_RAY_TRACING_POSITION_FETCH	1
 
@@ -49,8 +48,6 @@ public:
 
 	static void CreateSurface(FVulkanPlatformWindowContext& WindowContext, VkInstance Instance, VkSurfaceKHR* OutSurface);
 
-	static void WriteCrashMarker(const FOptionalVulkanDeviceExtensions& OptionalExtensions, FVulkanCommandBuffer* CmdBuffer, VkBuffer DestBuffer, const TArrayView<uint32>& Entries, bool bAdding);
-
 	static VkTimeDomainKHR GetTimeDomain() { return VK_TIME_DOMAIN_CLOCK_MONOTONIC_KHR; }
 
 protected:
diff --git a/Engine/Source/Runtime/VulkanRHI/Private/VulkanDevice.cpp b/Engine/Source/Runtime/VulkanRHI/Private/VulkanDevice.cpp
--- a/Engine/Source/Runtime/VulkanRHI/Private/VulkanDevice.cpp
+++ b/Engine/Source/Runtime/VulkanRHI/Private/VulkanDevice.cpp
@@ -1366,30 +1366,6 @@ void FVulkanDevice::InitGPU()
 		TempBlockAllocator = new VulkanRHI::FTempBlockAllocator(this, VulkanTempBlockSize, BlockAlignment, BufferUsageFlags);
 	}
 
-#if VULKAN_SUPPORTS_GPU_CRASH_DUMPS
-	if (UE::RHI::UseGPUCrashDebugging())
-	{
-		VkBufferCreateInfo CreateInfo;
-		ZeroVulkanStruct(CreateInfo, VK_STRUCTURE_TYPE_BUFFER_CREATE_INFO);
-		CreateInfo.size = GMaxCrashBufferEntries * sizeof(uint32_t);
-		CreateInfo.usage = VK_BUFFER_USAGE_TRANSFER_DST_BIT;
-		VERIFYVULKANRESULT(VulkanRHI::vkCreateBuffer(Device, &CreateInfo, VULKAN_CPU_ALLOCATOR, &CrashMarker.Buffer));
-
-		VkMemoryRequirements MemReq;
-		FMemory::Memzero(MemReq);
-		VulkanRHI::vkGetBufferMemoryRequirements(Device, CrashMarker.Buffer, &MemReq);
-
-		CrashMarker.Allocation = DeviceMemoryManager.Alloc(false, CreateInfo.size, MemReq.memoryTypeBits, VK_MEMORY_PROPERTY_HOST_VISIBLE_BIT |
-			VK_MEMORY_PROPERTY_HOST_COHERENT_BIT, nullptr, VULKAN_MEMORY_MEDIUM_PRIORITY, false, __FILE__, __LINE__);
-
-		uint32* Entry = (uint32*)CrashMarker.Allocation->Map(VK_WHOLE_SIZE, 0);
-		check(Entry);
-		// Start with 0 entries
-		*Entry = 0;
-		VERIFYVULKANRESULT(VulkanRHI::vkBindBufferMemory(Device, CrashMarker.Buffer, CrashMarker.Allocation->GetHandle(), 0));
-	}
-#endif
-
 	RenderPassManager = new FVulkanRenderPassManager(*this);
 
 	if (UseVulkanDescriptorCache())
@@ -1490,23 +1466,6 @@ void FVulkanDevice::Destroy()
 	PipelineStateCache = nullptr;
 	StagingManager.Deinit();
 
-#if VULKAN_SUPPORTS_GPU_CRASH_DUMPS
-	if (UE::RHI::UseGPUCrashDebugging())
-	{
-		if (CrashMarker.Buffer != VK_NULL_HANDLE)
-		{
-			VulkanRHI::vkDestroyBuffer(Device, CrashMarker.Buffer, VULKAN_CPU_ALLOCATOR);
-			CrashMarker.Buffer = VK_NULL_HANDLE;
-		}
-
-		if (CrashMarker.Allocation)
-		{
-			CrashMarker.Allocation->Unmap();
-			DeviceMemoryManager.Free(CrashMarker.Allocation);
-		}
-	}
-#endif // VULKAN_SUPPORTS_GPU_CRASH_DUMPS
-
 	DeferredDeletionQueue.Clear();
 
 	BindlessDescriptorManager->Deinit();
@@ -1524,14 +1483,13 @@ void FVulkanDevice::Destroy()
 	delete TempBlockAllocator;
 	TempBlockAllocator = nullptr;
 
-	MemoryManager.Deinit();
-
 	for (int32 QueueIndex = 0; QueueIndex < (int32)EVulkanQueueType::Count; ++QueueIndex)
 	{
 		delete Queues[QueueIndex];
 		Queues[QueueIndex] = nullptr;
 	}
 
+	MemoryManager.Deinit();
 	FenceManager.Deinit();
 	DeviceMemoryManager.Deinit();
 	FVulkanChunkedPipelineCacheManager::Shutdown();
@@ -1695,18 +1653,12 @@ FGPUTimingCalibrationTimestamp FVulkanDevice::GetCalibrationTimestamp()
 
 void FVulkanDevice::ForEachQueue(TFunctionRef<void(FVulkanQueue&)> Callback)
 {
-	static_assert((int32)EVulkanQueueType::Count == 3u, "New entries in EVulkanQueueType must be added in FVulkanDevice::ForEachQueue.");
-
-	Callback(*GetGraphicsQueue());
-
-	if (HasAsyncComputeQueue())
-	{
-		Callback(*GetComputeQueue());
-	}
-
-	if (HasTransferQueue())
+	for (int32 QueueIndex = 0; QueueIndex < Queues.Num(); ++QueueIndex)
 	{
-		Callback(*GetTransferQueue());
+		if (Queues[QueueIndex])
+		{
+			Callback(*Queues[QueueIndex]);
+		}
 	}
 }
 
diff --git a/Engine/Source/Runtime/VulkanRHI/Private/VulkanDevice.h b/Engine/Source/Runtime/VulkanRHI/Private/VulkanDevice.h
--- a/Engine/Source/Runtime/VulkanRHI/Private/VulkanDevice.h
+++ b/Engine/Source/Runtime/VulkanRHI/Private/VulkanDevice.h
@@ -538,18 +538,6 @@ public:
 		return OptionalDeviceExtensions.HasSeparateDepthStencilLayouts && OptionalDeviceExtensions.HasKHRSynchronization2 && OptionalDeviceExtensions.HasKHRRenderPass2;
 	}
 
-#if VULKAN_SUPPORTS_GPU_CRASH_DUMPS
-	VkBuffer GetCrashMarkerBuffer() const
-	{
-		return CrashMarker.Buffer;
-	}
-
-	void* GetCrashMarkerMappedPointer() const
-	{
-		return (CrashMarker.Allocation != nullptr) ? CrashMarker.Allocation->GetMappedPointer() : nullptr;
-	}
-#endif
-
 	void SetupPresentQueue(VkSurfaceKHR Surface);
 
 	const TArray<VkQueueFamilyProperties>& GetQueueFamilyProps()
@@ -658,14 +646,6 @@ private:
 
 	EGpuVendorId VendorId = EGpuVendorId::NotQueried;
 
-#if VULKAN_SUPPORTS_GPU_CRASH_DUMPS
-	struct
-	{
-		VkBuffer Buffer = VK_NULL_HANDLE;
-		VulkanRHI::FDeviceMemoryAllocation* Allocation = nullptr;
-	} CrashMarker;
-#endif
-
 	VkComponentMapping PixelFormatComponentMapping[PF_MAX];
 
 	TMap<uint32, FSamplerStateRHIRef> SamplerMap;
diff --git a/Engine/Source/Runtime/VulkanRHI/Private/VulkanExtensions.cpp b/Engine/Source/Runtime/VulkanRHI/Private/VulkanExtensions.cpp
--- a/Engine/Source/Runtime/VulkanRHI/Private/VulkanExtensions.cpp
+++ b/Engine/Source/Runtime/VulkanRHI/Private/VulkanExtensions.cpp
@@ -1169,7 +1169,7 @@ class FVulkanAMDBufferMarkerExtension : public FVulkanDeviceExtension
 public:
 
 	FVulkanAMDBufferMarkerExtension(FVulkanDevice* InDevice)
-		: FVulkanDeviceExtension(InDevice, VK_AMD_BUFFER_MARKER_EXTENSION_NAME, VULKAN_SUPPORTS_AMD_BUFFER_MARKER)
+		: FVulkanDeviceExtension(InDevice, VK_AMD_BUFFER_MARKER_EXTENSION_NAME, VULKAN_EXTENSION_ENABLED)
 	{
 		bEnabledInCode = bEnabledInCode && UE::RHI::UseGPUCrashDebugging() && UE::RHICore::AllowVendorDevice();
 	}
@@ -1188,7 +1188,7 @@ class FVulkanNVDeviceDiagnosticCheckpointsExtension : public FVulkanDeviceExtens
 public:
 
 	FVulkanNVDeviceDiagnosticCheckpointsExtension(FVulkanDevice* InDevice)
-		: FVulkanDeviceExtension(InDevice, VK_NV_DEVICE_DIAGNOSTIC_CHECKPOINTS_EXTENSION_NAME, VULKAN_SUPPORTS_NV_DIAGNOSTICS)
+		: FVulkanDeviceExtension(InDevice, VK_NV_DEVICE_DIAGNOSTIC_CHECKPOINTS_EXTENSION_NAME, VULKAN_EXTENSION_ENABLED)
 	{
 		bEnabledInCode = bEnabledInCode && UE::RHI::UseGPUCrashDebugging() && UE::RHICore::AllowVendorDevice();
 	}
@@ -1207,7 +1207,7 @@ class FVulkanNVDeviceDiagnosticConfigExtension : public FVulkanDeviceExtension
 public:
 
 	FVulkanNVDeviceDiagnosticConfigExtension(FVulkanDevice* InDevice)
-		: FVulkanDeviceExtension(InDevice, VK_NV_DEVICE_DIAGNOSTICS_CONFIG_EXTENSION_NAME, VULKAN_SUPPORTS_NV_DIAGNOSTICS)
+		: FVulkanDeviceExtension(InDevice, VK_NV_DEVICE_DIAGNOSTICS_CONFIG_EXTENSION_NAME, VULKAN_EXTENSION_ENABLED)
 	{
 		bEnabledInCode = bEnabledInCode && UE::RHI::UseGPUCrashDebugging() && UE::RHICore::AllowVendorDevice();
 	}
diff --git a/Engine/Source/Runtime/VulkanRHI/Private/VulkanGPUProfiler.h b/Engine/Source/Runtime/VulkanRHI/Private/VulkanGPUProfiler.h
--- a/Engine/Source/Runtime/VulkanRHI/Private/VulkanGPUProfiler.h
+++ b/Engine/Source/Runtime/VulkanRHI/Private/VulkanGPUProfiler.h
@@ -200,18 +200,6 @@ struct FVulkanGPUProfiler : public FGPUProfiler
 	FVulkanDevice* Device;
 	FVulkanContextCommon* CmdContext;
 
-#if VULKAN_SUPPORTS_GPU_CRASH_DUMPS
-	void PrepareCrashMarkerkBuffer();
-	void PushMarkerForCrash(FVulkanCommandBuffer* CmdBuffer, VkBuffer DestBuffer, const TCHAR* Name);
-	void PopMarkerForCrash(FVulkanCommandBuffer* CmdBuffer, VkBuffer DestBuffer);
-	void DumpCrashMarkers(void* BufferData);
-#endif
-
-	// For crash/marker tracking
-	TMap<uint32, FString> CachedStrings;
-	TArray<uint32> PushPopStack;
-
-	TArray<uint64> CrashMarkers;
 	bool bBeginFrame;
 };
 
diff --git a/Engine/Source/Runtime/VulkanRHI/Private/VulkanGenericPlatform.cpp b/Engine/Source/Runtime/VulkanRHI/Private/VulkanGenericPlatform.cpp
--- a/Engine/Source/Runtime/VulkanRHI/Private/VulkanGenericPlatform.cpp
+++ b/Engine/Source/Runtime/VulkanRHI/Private/VulkanGenericPlatform.cpp
@@ -185,23 +185,3 @@ FString FVulkanGenericPlatform::GetVulkanProfileNameForFeatureLevel(ERHIFeatureL
 	}
 	return ProfileName;
 }
-
-void FVulkanGenericPlatform::WriteCrashMarkerWithoutExtensions(FVulkanCommandBuffer* CmdBuffer, VkBuffer DestBuffer, const TArrayView<uint32>& Entries, bool bAdding)
-{
-	// Buffer writes cannot be recorded inside of a render pass
-	if (CmdBuffer->IsOutsideRenderPass())
-	{
-		// execution barrier
-		VulkanDynamicAPI::vkCmdPipelineBarrier(CmdBuffer->GetHandle(), VK_PIPELINE_STAGE_BOTTOM_OF_PIPE_BIT, VK_PIPELINE_STAGE_TOP_OF_PIPE_BIT, 0, 0, nullptr, 0, nullptr, 0, nullptr);
-
-		VulkanDynamicAPI::vkCmdFillBuffer(CmdBuffer->GetHandle(), DestBuffer, 0, sizeof(uint32), Entries.Num());
-		if (bAdding)
-		{
-			int32 LastIndex = Entries.Num() - 1;
-			VulkanDynamicAPI::vkCmdFillBuffer(CmdBuffer->GetHandle(), DestBuffer, (1 + LastIndex) * sizeof(uint32), sizeof(uint32), Entries[LastIndex]);
-		}
-
-		// execution barrier
-		VulkanDynamicAPI::vkCmdPipelineBarrier(CmdBuffer->GetHandle(), VK_PIPELINE_STAGE_BOTTOM_OF_PIPE_BIT, VK_PIPELINE_STAGE_TOP_OF_PIPE_BIT, 0, 0, nullptr, 0, nullptr, 0, nullptr);
-	}
-}
diff --git a/Engine/Source/Runtime/VulkanRHI/Private/VulkanGenericPlatform.h b/Engine/Source/Runtime/VulkanRHI/Private/VulkanGenericPlatform.h
--- a/Engine/Source/Runtime/VulkanRHI/Private/VulkanGenericPlatform.h
+++ b/Engine/Source/Runtime/VulkanRHI/Private/VulkanGenericPlatform.h
@@ -93,10 +93,6 @@ public:
 
 	static bool RegisterGPUWork() { return true; }
 
-	static void WriteCrashMarker(const FOptionalVulkanDeviceExtensions& OptionalExtensions, FVulkanCommandBuffer* CmdBuffer, VkBuffer DestBuffer, const TArrayView<uint32>& Entries, bool bAdding) {}
-
-	static void WriteCrashMarkerWithoutExtensions(FVulkanCommandBuffer* CmdBuffer, VkBuffer DestBuffer, const TArrayView<uint32>& Entries, bool bAdding);
-
 	// Allow the platform code to restrict the device features
 	static void RestrictEnabledPhysicalDeviceFeatures(FVulkanPhysicalDeviceFeatures* InOutFeaturesToEnable);
 
diff --git a/Engine/Source/Runtime/VulkanRHI/Private/VulkanRHI.cpp b/Engine/Source/Runtime/VulkanRHI/Private/VulkanRHI.cpp
--- a/Engine/Source/Runtime/VulkanRHI/Private/VulkanRHI.cpp
+++ b/Engine/Source/Runtime/VulkanRHI/Private/VulkanRHI.cpp
@@ -1335,12 +1335,6 @@ void FVulkanCommandListContext::RHIEndDrawingViewport(FRHIViewport* ViewportRHI,
 #else
 		if (IsImmediate())
 		{
-		#if VULKAN_SUPPORTS_GPU_CRASH_DUMPS
-			if (GpuProfiler.bTrackingGPUCrashData)
-			{
-				GpuProfiler.PushMarkerForCrash(GetActiveCmdBuffer(), Device.GetCrashMarkerBuffer(), GetNameStr());
-			}
-		#endif
 			if (GpuProfiler.IsProfilingGPU())
 			{
 				GpuProfiler.PushEvent(GetNameStr(), Color);
@@ -1369,13 +1363,6 @@ void FVulkanCommandListContext::RHIEndDrawingViewport(FRHIViewport* ViewportRHI,
 			{
 				GpuProfiler.PopEvent();
 			}
-
-		#if VULKAN_SUPPORTS_GPU_CRASH_DUMPS
-			if (GpuProfiler.bTrackingGPUCrashData)
-			{
-				GpuProfiler.PopMarkerForCrash(GetActiveCmdBuffer(), Device.GetCrashMarkerBuffer());
-			}
-		#endif
 		}
 #endif // RHI_NEW_GPU_PROFILER
 
diff --git a/Engine/Source/Runtime/VulkanRHI/Private/VulkanUtil.cpp b/Engine/Source/Runtime/VulkanRHI/Private/VulkanUtil.cpp
--- a/Engine/Source/Runtime/VulkanRHI/Private/VulkanUtil.cpp
+++ b/Engine/Source/Runtime/VulkanRHI/Private/VulkanUtil.cpp
@@ -308,26 +308,6 @@ FVulkanGPUProfiler::~FVulkanGPUProfiler()
 
 void FVulkanGPUProfiler::BeginFrame()
 {
-#if VULKAN_SUPPORTS_GPU_CRASH_DUMPS
-	if (UE::RHI::UseGPUCrashDebugging())
-	{
-		static auto* CrashCollectionEnableCvar = IConsoleManager::Get().FindTConsoleVariableDataInt(TEXT("r.gpucrash.collectionenable"));
-		static auto* CrashCollectionDataDepth = IConsoleManager::Get().FindTConsoleVariableDataInt(TEXT("r.gpucrash.datadepth"));
-		bTrackingGPUCrashData = CrashCollectionEnableCvar ? CrashCollectionEnableCvar->GetValueOnRenderThread() != 0 : false;
-		GPUCrashDataDepth = CrashCollectionDataDepth ? CrashCollectionDataDepth->GetValueOnRenderThread() : -1;
-		if (GPUCrashDataDepth == -1 || GPUCrashDataDepth > GMaxCrashBufferEntries)
-		{
-			static bool bChecked = false;
-			if (!bChecked)
-			{
-				bChecked = true;
-				UE_LOG(LogVulkanRHI, Warning, TEXT("Clamping r.gpucrash.datadepth to %d"), GMaxCrashBufferEntries);
-			}
-			GPUCrashDataDepth = GMaxCrashBufferEntries;
-		}
-	}
-#endif
-
 	bCommandlistSubmitted = false;
 	CurrentEventNode = NULL;
 	check(!bTrackingEvents);
@@ -429,119 +409,9 @@ void FVulkanGPUProfiler::EndFrame()
 
 	bBeginFrame = false;
 }
-
-#if VULKAN_SUPPORTS_GPU_CRASH_DUMPS
-
-void FVulkanGPUProfiler::PushMarkerForCrash(FVulkanCommandBuffer* CmdBuffer, VkBuffer DestBuffer, const TCHAR* Name)
-{
-	uint32 CRC = 0;
-	if (GPUCrashDataDepth < 0 || PushPopStack.Num() < GPUCrashDataDepth)
-	{
-		CRC = FCrc::StrCrc32<TCHAR>(Name);
-
-		if (CachedStrings.Num() > 10000)
-		{
-			CachedStrings.Empty(10000);
-			CachedStrings.Emplace(EventDeepCRC, EventDeepString);
-		}
-
-		if (CachedStrings.Find(CRC) == nullptr)
-		{
-			CachedStrings.Emplace(CRC, FString(Name));
-		}
-	}
-	else
-	{
-		CRC = EventDeepCRC;
-	}
-
-	PushPopStack.Push(CRC);
-	FVulkanPlatform::WriteCrashMarker(Device->GetOptionalExtensions(), CmdBuffer, DestBuffer, TArrayView<uint32>(PushPopStack), true);
-}
-
-void FVulkanGPUProfiler::PopMarkerForCrash(FVulkanCommandBuffer* CmdBuffer, VkBuffer DestBuffer)
-{
-	if (PushPopStack.Num() > 0)
-	{
-		PushPopStack.Pop(EAllowShrinking::No);
-		FVulkanPlatform::WriteCrashMarker(Device->GetOptionalExtensions(), CmdBuffer, DestBuffer, TArrayView<uint32>(PushPopStack), false);
-	}
-}
-
-void FVulkanGPUProfiler::DumpCrashMarkers(void* BufferData)
-{
-	{
-		if (Device->GetOptionalExtensions().HasAMDBufferMarker)
-		{
-			UE_LOG(LogVulkanRHI, Error, TEXT("Breadcrumbs using VK_AMD_buffer_marker extension"));
-		}
-		else
-		{
-			UE_LOG(LogVulkanRHI, Error, TEXT("Breadcrumbs without extensions"));
-		}
-
-		uint32* Entries = (uint32*)BufferData;
-		uint32 NumCRCs = *Entries++;
-		for (uint32 Index = 0; Index < NumCRCs; ++Index)
-		{
-			const FString* Frame = CachedStrings.Find(*Entries);
-			UE_LOG(LogVulkanRHI, Error, TEXT("[GPU Breadcrumb] %i: %s (CRC 0x%x)"), Index, Frame ? *(*Frame) : TEXT("<undefined>"), *Entries);
-			++Entries;
-		}
-	}
-
-	if (Device->GetOptionalExtensions().HasNVDiagnosticCheckpoints)
-	{
-		struct FCheckpointDataNV : public VkCheckpointDataNV
-		{
-			FCheckpointDataNV()
-			{
-				ZeroVulkanStruct(*this, VK_STRUCTURE_TYPE_CHECKPOINT_DATA_NV);
-			}
-		};
-		TArray<FCheckpointDataNV> Data;
-		uint32 Num = 0;
-		VkQueue QueueHandle = Device->GetGraphicsQueue()->GetHandle();
-		VulkanDynamicAPI::vkGetQueueCheckpointDataNV(QueueHandle, &Num, nullptr);
-		if (Num > 0)
-		{
-			Data.AddDefaulted(Num);
-			VulkanDynamicAPI::vkGetQueueCheckpointDataNV(QueueHandle, &Num, &Data[0]);
-			check(Num == Data.Num());
-			for (uint32 Index = 0; Index < Num; ++Index)
-			{
-				check(Data[Index].sType == VK_STRUCTURE_TYPE_CHECKPOINT_DATA_NV);
-				uint32 Value = (uint32)(size_t)Data[Index].pCheckpointMarker;
-				const FString* Frame = CachedStrings.Find(Value);
-				UE_LOG(LogVulkanRHI, Error, TEXT("[VK_NV_device_diagnostic_checkpoints] %i: Stage %s (0x%08x), %s (CRC 0x%x)"), 
-					Index, VK_TYPE_TO_STRING(VkPipelineStageFlagBits, Data[Index].stage), Data[Index].stage, Frame ? *(*Frame) : TEXT("<undefined>"), Value);
-			}
-			GLog->Panic();
-		}
-	}
-}
-#endif // VULKAN_SUPPORTS_GPU_CRASH_DUMPS
+#endif // (RHI_NEW_GPU_PROFILER == 0)
 
 #if NV_AFTERMATH
-void AftermathResolveMarkerCallback(const void* Marker, void** ResolvedMarkerData, uint32_t* MarkerSize)
-{
-#if VULKAN_SUPPORTS_NV_DIAGNOSTICS
-	FVulkanDevice* VulkanDevice = FVulkanDynamicRHI::Get().GetDevice();
-	if (VulkanDevice->GetOptionalExtensions().HasNVDiagnosticCheckpoints)
-	{
-		const uint32 Value = (uint32)(size_t)Marker;
-		const FString* MarkerName = VulkanDevice->GetImmediateContext().GetGPUProfiler().CachedStrings.Find(Value);
-		UE_LOG(LogVulkanRHI, Display, TEXT("[AftermathResolveMarkerCallback] Requested %u [%s]"), Value, MarkerName ? *(*MarkerName) : TEXT("<undefined>"));
-		if (MarkerName && !MarkerName->IsEmpty() && ResolvedMarkerData && MarkerSize)
-		{
-			const TArray<TCHAR, FString::AllocatorType>& CharArray = MarkerName->GetCharArray();
-			(*ResolvedMarkerData) = (void*)CharArray.GetData();
-			(*MarkerSize) = CharArray.Num() * CharArray.GetTypeSize();
-		}
-	}
-#endif // VULKAN_SUPPORTS_NV_DIAGNOSTICS
-}
-
 void ConditionallyAddShader(const FVulkanShader* Value, TArray<const FVulkanShader*>& Shaders, TSet<const FVulkanShader*>& ShaderSet)
 {
 	if (Value && !ShaderSet.Contains(Value))
@@ -681,9 +551,8 @@ void AftermathLateAssociate(float TimeLimitSeconds, uint32 FrameLimit)
 	double TimeMS = FPlatformTime::ToMilliseconds64(FPlatformTime::Cycles64() - CycleStart);
 	UE_LOG(LogVulkanRHI, Log, TEXT("Created late shader associations, took %f.5ms"), TimeMS);
 }
-#endif
+#endif // NV_AFTERMATH
 
-#endif // (RHI_NEW_GPU_PROFILER == 0)
 
 FVulkanStagingBuffer::~FVulkanStagingBuffer()
 {
@@ -963,6 +832,7 @@ namespace VulkanRHI
 		VKERRORCASE(VK_ERROR_NOT_PERMITTED_EXT); break;
 #undef VKERRORCASE
 		default:
+			ErrorString = VK_TYPE_TO_STRING(VkResult, Result);
 			break;
 		}
 
@@ -985,15 +855,6 @@ namespace VulkanRHI
 
 #if RHI_NEW_GPU_PROFILER
 			checkNoEntry(); // @todo - new gpu profiler
-#else
-
-#if VULKAN_SUPPORTS_GPU_CRASH_DUMPS
-			if (UE::RHI::UseGPUCrashDebugging())
-			{
-				Device->GetImmediateContext().GetGPUProfiler().DumpCrashMarkers(Device->GetCrashMarkerMappedPointer());
-			}
-#endif
-
 #endif
 
 			CheckDeviceFault(Device);
diff --git a/Engine/Source/Runtime/VulkanRHI/Private/Windows/VulkanWindowsPlatform.cpp b/Engine/Source/Runtime/VulkanRHI/Private/Windows/VulkanWindowsPlatform.cpp
--- a/Engine/Source/Runtime/VulkanRHI/Private/Windows/VulkanWindowsPlatform.cpp
+++ b/Engine/Source/Runtime/VulkanRHI/Private/Windows/VulkanWindowsPlatform.cpp
@@ -194,34 +194,3 @@ bool FVulkanWindowsPlatform::SupportsDeviceLocalHostVisibleWithNoPenalty(EGpuVen
 	static bool bIsWin10 = FPlatformMisc::VerifyWindowsVersion(10, 0) /*Win10*/;
 	return (VendorId == EGpuVendorId::Amd && bIsWin10);
 }
-
-void FVulkanWindowsPlatform::WriteCrashMarker(const FOptionalVulkanDeviceExtensions& OptionalExtensions, FVulkanCommandBuffer* CmdBuffer, VkBuffer DestBuffer, const TArrayView<uint32>& Entries, bool bAdding)
-{
-	ensure(Entries.Num() <= GMaxCrashBufferEntries);
-
-	if (OptionalExtensions.HasAMDBufferMarker)
-	{
-		// AMD API only allows updating one entry at a time. Assume buffer has entry 0 as num entries
-		VulkanDynamicAPI::vkCmdWriteBufferMarkerAMD(CmdBuffer->GetHandle(), VK_PIPELINE_STAGE_BOTTOM_OF_PIPE_BIT, DestBuffer, 0, Entries.Num());
-		if (bAdding)
-		{
-			int32 LastIndex = Entries.Num() - 1;
-			// +1 size as entries start at index 1
-			VulkanDynamicAPI::vkCmdWriteBufferMarkerAMD(CmdBuffer->GetHandle(), VK_PIPELINE_STAGE_BOTTOM_OF_PIPE_BIT, DestBuffer, (1 + LastIndex) * sizeof(uint32), Entries[LastIndex]);
-		}
-	}
-	else
-	{
-		WriteCrashMarkerWithoutExtensions(CmdBuffer, DestBuffer, Entries, bAdding);
-	}
-
-	if (OptionalExtensions.HasNVDiagnosticCheckpoints)
-	{
-		if (bAdding)
-		{
-			int32 LastIndex = Entries.Num() - 1;
-			uint32 Value = Entries[LastIndex];
-			VulkanDynamicAPI::vkCmdSetCheckpointNV(CmdBuffer->GetHandle(), (void*)(size_t)Value);
-		}
-	}
-}
diff --git a/Engine/Source/Runtime/VulkanRHI/Private/Windows/VulkanWindowsPlatform.h b/Engine/Source/Runtime/VulkanRHI/Private/Windows/VulkanWindowsPlatform.h
--- a/Engine/Source/Runtime/VulkanRHI/Private/Windows/VulkanWindowsPlatform.h
+++ b/Engine/Source/Runtime/VulkanRHI/Private/Windows/VulkanWindowsPlatform.h
@@ -14,7 +14,6 @@
 #define VULKAN_DYNAMICALLYLOADED					1
 #define VULKAN_SHOULD_ENABLE_DESKTOP_HMD_SUPPORT	1
 #define VULKAN_SIGNAL_UNIMPLEMENTED()				checkf(false, TEXT("Unimplemented vulkan functionality: %s"), StringCast<TCHAR>(__FUNCTION__).Get())
-#define VULKAN_SUPPORTS_AMD_BUFFER_MARKER			1
 #define VULKAN_SUPPORTS_SCALAR_BLOCK_LAYOUT			1
 #define VULKAN_SUPPORTS_RAY_TRACING_POSITION_FETCH	1
 
@@ -58,8 +57,6 @@ public:
 
 	static bool SupportsDeviceLocalHostVisibleWithNoPenalty(EGpuVendorId VendorId);
 
-	static void WriteCrashMarker(const FOptionalVulkanDeviceExtensions& OptionalExtensions, FVulkanCommandBuffer* CmdBuffer, VkBuffer DestBuffer, const TArrayView<uint32>& Entries, bool bAdding);
-
 	static VkTimeDomainKHR GetTimeDomain() { return VK_TIME_DOMAIN_QUERY_PERFORMANCE_COUNTER_KHR; }
 
 private:
diff --git a/Engine/Source/Runtime/VulkanRHI/Public/VulkanConfiguration.h b/Engine/Source/Runtime/VulkanRHI/Public/VulkanConfiguration.h
--- a/Engine/Source/Runtime/VulkanRHI/Public/VulkanConfiguration.h
+++ b/Engine/Source/Runtime/VulkanRHI/Public/VulkanConfiguration.h
@@ -145,32 +145,6 @@ struct VkAllocationCallbacks;
 	#define VULKAN_DELETE_STALE_CMDBUFFERS						1
 #endif
 
-#ifndef VULKAN_SUPPORTS_AMD_BUFFER_MARKER
-	#define VULKAN_SUPPORTS_AMD_BUFFER_MARKER					0
-#endif
-
-#ifndef VULKAN_SUPPORTS_NV_DIAGNOSTIC_CHECKPOINT
-	#ifdef VK_NV_device_diagnostic_checkpoints
-		#define VULKAN_SUPPORTS_NV_DIAGNOSTIC_CHECKPOINT		1
-	#else
-		#define VULKAN_SUPPORTS_NV_DIAGNOSTIC_CHECKPOINT		0
-	#endif
-#endif
-
-#ifndef VULKAN_SUPPORTS_NV_DIAGNOSTIC_CONFIG
-	#ifdef VK_NV_device_diagnostics_config
-		#define VULKAN_SUPPORTS_NV_DIAGNOSTIC_CONFIG			1
-	#else
-		#define VULKAN_SUPPORTS_NV_DIAGNOSTIC_CONFIG			0
-	#endif
-#endif
-
-#define VULKAN_SUPPORTS_NV_DIAGNOSTICS							(VULKAN_SUPPORTS_NV_DIAGNOSTIC_CHECKPOINT && VULKAN_SUPPORTS_NV_DIAGNOSTIC_CONFIG)
-
-#ifndef VULKAN_SUPPORTS_GPU_CRASH_DUMPS
-	#define VULKAN_SUPPORTS_GPU_CRASH_DUMPS						(VULKAN_SUPPORTS_AMD_BUFFER_MARKER || VULKAN_SUPPORTS_NV_DIAGNOSTIC_CONFIG)
-#endif
-
 #ifndef VULKAN_SUPPORTS_SCALAR_BLOCK_LAYOUT
 	#define VULKAN_SUPPORTS_SCALAR_BLOCK_LAYOUT					1
 #endif
