VulkanRHI: clamp values in the calibration fallback

#rnx

diff --git a/Engine/Source/Runtime/VulkanRHI/Private/VulkanSubmission.cpp b/Engine/Source/Runtime/VulkanRHI/Private/VulkanSubmission.cpp
--- a/Engine/Source/Runtime/VulkanRHI/Private/VulkanSubmission.cpp
+++ b/Engine/Source/Runtime/VulkanRHI/Private/VulkanSubmission.cpp
@@ -654,6 +654,9 @@ void FVulkanDynamicRHI::CompletePayload(FVulkanPayload* Payload)
 
 								CPUCalibrationTimestamp = ReferenceCPUTimestamp;
 								GPUCalibrationTimestamp = ReferenceGPUTimestamp;
+
+								// :todo: Clamp results within the imperfect calibration timestamps (to prevent negative values below)
+								TempResults[QueryIndex].Result = FMath::Max(TempResults[QueryIndex].Result, GPUCalibrationTimestamp);
 							}
 
 							const uint64 GPUDelta = TempResults[QueryIndex].Result - GPUCalibrationTimestamp;
