VulkanRHI: Remove unused stat

#rnx

diff --git a/Engine/Source/Runtime/VulkanRHI/Private/VulkanRHIPrivate.h b/Engine/Source/Runtime/VulkanRHI/Private/VulkanRHIPrivate.h
--- a/Engine/Source/Runtime/VulkanRHI/Private/VulkanRHIPrivate.h
+++ b/Engine/Source/Runtime/VulkanRHI/Private/VulkanRHIPrivate.h
@@ -414,7 +414,6 @@ DECLARE_DWORD_ACCUMULATOR_STAT_EXTERN(TEXT("Num Frame Buffers"), STAT_VulkanNumF
 DECLARE_DWORD_ACCUMULATOR_STAT_EXTERN(TEXT("Num Buffer Views"), STAT_VulkanNumBufferViews, STATGROUP_VulkanRHI, );
 DECLARE_DWORD_ACCUMULATOR_STAT_EXTERN(TEXT("Num Image Views"), STAT_VulkanNumImageViews, STATGROUP_VulkanRHI, );
 DECLARE_DWORD_ACCUMULATOR_STAT_EXTERN(TEXT("Num Physical Mem Allocations"), STAT_VulkanNumPhysicalMemAllocations, STATGROUP_VulkanRHI, );
-DECLARE_MEMORY_STAT_EXTERN(TEXT("Frame Temp Memory"), STAT_VulkanTempFrameAllocationBuffer, STATGROUP_VulkanRHI, );
 DECLARE_DWORD_COUNTER_STAT_EXTERN(TEXT("Dynamic VB Size"), STAT_VulkanDynamicVBSize, STATGROUP_VulkanRHI, );
 DECLARE_DWORD_COUNTER_STAT_EXTERN(TEXT("Dynamic IB Size"), STAT_VulkanDynamicIBSize, STATGROUP_VulkanRHI, );
 DECLARE_CYCLE_STAT_EXTERN(TEXT("Dynamic VB Lock/Unlock time"), STAT_VulkanDynamicVBLockTime, STATGROUP_VulkanRHI, );
diff --git a/Engine/Source/Runtime/VulkanRHI/Private/VulkanUtil.cpp b/Engine/Source/Runtime/VulkanRHI/Private/VulkanUtil.cpp
--- a/Engine/Source/Runtime/VulkanRHI/Private/VulkanUtil.cpp
+++ b/Engine/Source/Runtime/VulkanRHI/Private/VulkanUtil.cpp
@@ -913,7 +913,6 @@ DEFINE_STAT(STAT_VulkanNumFrameBuffers);
 DEFINE_STAT(STAT_VulkanNumBufferViews);
 DEFINE_STAT(STAT_VulkanNumImageViews);
 DEFINE_STAT(STAT_VulkanNumPhysicalMemAllocations);
-DEFINE_STAT(STAT_VulkanTempFrameAllocationBuffer);
 DEFINE_STAT(STAT_VulkanDynamicVBSize);
 DEFINE_STAT(STAT_VulkanDynamicIBSize);
 DEFINE_STAT(STAT_VulkanDynamicVBLockTime);
