VulkanRHI: Turn off VK_FLAGS_TO_STRING on Android.

#rnx

diff --git a/Engine/Source/Runtime/VulkanRHI/Private/Android/VulkanAndroidPlatform.h b/Engine/Source/Runtime/VulkanRHI/Private/Android/VulkanAndroidPlatform.h
--- a/Engine/Source/Runtime/VulkanRHI/Private/Android/VulkanAndroidPlatform.h
+++ b/Engine/Source/Runtime/VulkanRHI/Private/Android/VulkanAndroidPlatform.h
@@ -30,7 +30,7 @@
 #if UE_BUILD_DEBUG || UE_BUILD_DEVELOPMENT
 #	include "vk_enum_string_helper.h"
 #	define VK_TYPE_TO_STRING(Type, Value) ANSI_TO_TCHAR(string_##Type(Value))
-#	define VK_FLAGS_TO_STRING(Type, Value) ANSI_TO_TCHAR(string_##Type(Value).c_str())
+//#	define VK_FLAGS_TO_STRING(Type, Value) ANSI_TO_TCHAR(string_##Type(Value).c_str())
 #endif
 
 #define ENUM_VK_ENTRYPOINTS_PLATFORM_BASE(EnumMacro)
