VulkanRHI: Bring back crash handling in VulkanRHI with new profiler:
- Execute crash handler only from interrupt thread (or whichever thread is handling payload completion).
- Bring back vendor buffer markers and diagnostic checkpoints (use them with breadcrumbs).
- Extend breadcrumb lifetimes when using diagnostic checkpoints.
- Add Vulkan implementation for FRHIDiagnosticBuffer (currently only used for breadcrumbs, shader asserts to come later).
- Dump breadcrumb tree on crash.
- Only pop payloads from pending interrupts after it was completed in case the crash occurs during payload completion.

#rnx

diff --git a/Engine/Config/Android/AndroidEngine.ini b/Engine/Config/Android/AndroidEngine.ini
--- a/Engine/Config/Android/AndroidEngine.ini
+++ b/Engine/Config/Android/AndroidEngine.ini
@@ -96,6 +96,8 @@ r.Vulkan.AllowSplitBarriers=0
 ; Split LensFlareBlur into two passes to remove vertex shader texture sample
 ; Splitting into compute shader is necessary on Android. See UE-243673
 r.LensFlareBlurComputeShader=1
+; Crash reporting diagnotic buffer
+r.Vulkan.DiagnosticBuffer=1
 
 [/Script/Engine.RendererSettings]
 r.DistanceFields=1
diff --git a/Engine/Source/Runtime/RHICore/Private/RHICoreNvidiaAftermath.cpp b/Engine/Source/Runtime/RHICore/Private/RHICoreNvidiaAftermath.cpp
--- a/Engine/Source/Runtime/RHICore/Private/RHICoreNvidiaAftermath.cpp
+++ b/Engine/Source/Runtime/RHICore/Private/RHICoreNvidiaAftermath.cpp
@@ -132,7 +132,6 @@ namespace UE::RHICore::Nvidia::Aftermath
 	static void* DllHandle = nullptr;
 	static bool bEnabled = false;
 	static uint32 Flags = 0;
-	static FResolveMarkerFunc GResolveMarkerFunc;
 	static FAssociateShadersFunc GLateShaderAssociateFunc;
 
 	struct FShaderDebugInfo
@@ -212,12 +211,10 @@ namespace UE::RHICore::Nvidia::Aftermath
 		GLateShaderAssociateFunc = MoveTemp(ShaderAssociateFunc);
 	}
 
-	RHICORE_API void InitializeBeforeDeviceCreation(FResolveMarkerFunc ResolveMarkerFunc)
+	RHICORE_API void InitializeBeforeDeviceCreation(bool bResolveMarkers)
 	{
 		LoadAftermathDLL();
 
-		GResolveMarkerFunc = MoveTemp(ResolveMarkerFunc);
-
 		if (!DllHandle)
 			return;
 
@@ -237,8 +234,6 @@ namespace UE::RHICore::Nvidia::Aftermath
 		Flags |= (FParse::Param(FCommandLine::Get(), TEXT("nvAftermathDumpShaderDebugInfo")) || (CVarAftermath_DumpShaderDebugInfo->GetInt())) ? GFSDK_Aftermath_FeatureFlags_GenerateShaderDebugInfo : 0;
 		Flags |= (FParse::Param(FCommandLine::Get(), TEXT("nvaftermathall"                )) || (CVarAftermath_TrackAll           ->GetInt())) ? AllFlags                                             : 0;
 
-		const bool bResolveMarkers = NV_AFTERMATH_USE_BREADCRUMB_PTRS || GResolveMarkerFunc.IsSet();
-		
 		GFSDK_Aftermath_Result Result = GFSDK_Aftermath_EnableGpuCrashDumps(
 			GFSDK_Aftermath_Version_API,
 #if PLATFORM_WINDOWS
@@ -846,37 +841,29 @@ namespace UE::RHICore::Nvidia::Aftermath
 
 	static void GFSDK_AFTERMATH_CALL Callback_ResolveMarker(const void* MarkerData, const uint32_t MarkerDataSize, void* UserData, void** ResolvedMarkerData, uint32_t* ResolvedMarkerDataSize)
 	{
-		if (GResolveMarkerFunc)
+	#if WITH_RHI_BREADCRUMBS
+		FRHIBreadcrumbNode const* Breadcrumb = static_cast<FRHIBreadcrumbNode const*>(MarkerData);
+		if (Breadcrumb == FRHIBreadcrumbNode::Sentinel)
 		{
-			// Some RHIs override the handling of Aftermath markers
-			GResolveMarkerFunc(MarkerData, MarkerDataSize, UserData, ResolvedMarkerData, ResolvedMarkerDataSize);
+			*ResolvedMarkerData = const_cast<TCHAR*>(RootNodeName);
+			*ResolvedMarkerDataSize = sizeof(RootNodeName);
 		}
 		else
 		{
-		#if WITH_RHI_BREADCRUMBS
-			FRHIBreadcrumbNode const* Breadcrumb = static_cast<FRHIBreadcrumbNode const*>(MarkerData);
-			if (Breadcrumb == FRHIBreadcrumbNode::Sentinel)
-			{
-				*ResolvedMarkerData = const_cast<TCHAR*>(RootNodeName);
-				*ResolvedMarkerDataSize = sizeof(RootNodeName);
-			}
-			else
-			{
-				// Allocate space to hold the name of this breadcrumb.
-				// The memory must remain valid until Aftermath calls the GpuCrashDump callback.
-				FRHIBreadcrumb::FBuffer* Buffer = new FRHIBreadcrumb::FBuffer;
-				NameStorage.Add(Buffer);
+			// Allocate space to hold the name of this breadcrumb.
+			// The memory must remain valid until Aftermath calls the GpuCrashDump callback.
+			FRHIBreadcrumb::FBuffer* Buffer = new FRHIBreadcrumb::FBuffer;
+			NameStorage.Add(Buffer);
 
-				TCHAR const* NameStr = Breadcrumb->GetTCHAR(*Buffer);
+			TCHAR const* NameStr = Breadcrumb->GetTCHAR(*Buffer);
 
-				*ResolvedMarkerData = const_cast<TCHAR*>(NameStr);
-				*ResolvedMarkerDataSize = (FCString::Strlen(NameStr) + 1) * sizeof(TCHAR); // Include null terminator
-			}
-		#else
-			*ResolvedMarkerData = const_cast<TCHAR*>(BreadcrumbsDisabledStr);
-			*ResolvedMarkerDataSize = (FCString::Strlen(BreadcrumbsDisabledStr) + 1) * sizeof(TCHAR); // Include null terminator
-		#endif
+			*ResolvedMarkerData = const_cast<TCHAR*>(NameStr);
+			*ResolvedMarkerDataSize = (FCString::Strlen(NameStr) + 1) * sizeof(TCHAR); // Include null terminator
 		}
+	#else
+		*ResolvedMarkerData = const_cast<TCHAR*>(BreadcrumbsDisabledStr);
+		*ResolvedMarkerDataSize = (FCString::Strlen(BreadcrumbsDisabledStr) + 1) * sizeof(TCHAR); // Include null terminator
+	#endif
 	}
 
 	void Callback_ShaderDebugInfo(const void* ShaderDebugInfo, const uint32 ShaderDebugInfoSize, void* UserData)
diff --git a/Engine/Source/Runtime/RHICore/Public/RHICoreNvidiaAftermath.h b/Engine/Source/Runtime/RHICore/Public/RHICoreNvidiaAftermath.h
--- a/Engine/Source/Runtime/RHICore/Public/RHICoreNvidiaAftermath.h
+++ b/Engine/Source/Runtime/RHICore/Public/RHICoreNvidiaAftermath.h
@@ -18,7 +18,6 @@
 
 	namespace UE::RHICore::Nvidia::Aftermath
 	{
-		typedef TUniqueFunction<void(const void* MarkerData, const uint32_t MarkerDataSize, void* UserData, void** ResolvedMarkerData, uint32_t* ResolvedMarkerDataSize)> FResolveMarkerFunc;
 		typedef TUniqueFunction<void(float TimeLimitSeconds, uint32 FrameLimit)> FAssociateShadersFunc;
 		
 		RHICORE_API bool IsEnabled();
@@ -29,7 +28,7 @@
 		// Called by platform RHIs to activate Aftermath.
 		// The optional ResolveMarkerFunc is provided for platform RHIs to replace the default marker resolution with something custom.
 		//
-		RHICORE_API void InitializeBeforeDeviceCreation(FResolveMarkerFunc ResolveMarkerFunc = {});
+		RHICORE_API void InitializeBeforeDeviceCreation(bool bResolveMarkers = NV_AFTERMATH_USE_BREADCRUMB_PTRS);
 
 		//
 		// Called by the crash handler to associate shaders after a GPU fault
diff --git a/Engine/Source/Runtime/VulkanRHI/Private/VulkanDevice.cpp b/Engine/Source/Runtime/VulkanRHI/Private/VulkanDevice.cpp
--- a/Engine/Source/Runtime/VulkanRHI/Private/VulkanDevice.cpp
+++ b/Engine/Source/Runtime/VulkanRHI/Private/VulkanDevice.cpp
@@ -1381,6 +1381,14 @@ void FVulkanDevice::InitGPU()
 
 	InitGlobalSamplers();
 
+	if (UE::RHI::UseGPUCrashBreadcrumbs())
+	{
+		ForEachQueue([](FVulkanQueue& Queue)
+			{
+				Queue.InitDiagnosticBuffer();
+			});
+	}
+
 	PipelineStateCache = new FVulkanPipelineStateCacheManager(this);
 
 	TArray<FString> CacheFilenames = FVulkanPlatform::GetPSOCacheFilenames();
diff --git a/Engine/Source/Runtime/VulkanRHI/Private/VulkanDynamicRHI.h b/Engine/Source/Runtime/VulkanRHI/Private/VulkanDynamicRHI.h
--- a/Engine/Source/Runtime/VulkanRHI/Private/VulkanDynamicRHI.h
+++ b/Engine/Source/Runtime/VulkanRHI/Private/VulkanDynamicRHI.h
@@ -326,13 +326,22 @@ public:
 	// Runs code on SubmissionThread with access to VkQueue.  Useful for plugins. 
 	virtual void RHIRunOnQueue(EVulkanRHIRunOnQueueType QueueType, TFunction<void(VkQueue)>&& CodeToRun, bool bWaitForSubmission) final override;
 
-public:
+	void DumpCrashMarkers();
+
 	static void SavePipelineCache();
 
 	void ProcessInterruptQueueUntil(FGraphEvent* GraphEvent);
 	void EnqueueEndOfPipeTask(TUniqueFunction<void()> TaskFunc, TUniqueFunction<void(FVulkanPayload&)> ModifyPayloadCallback = {});
 	void CompletePayload(FVulkanPayload* Payload);
 
+	// Returns true if the calling thread is the interrupt thread
+	// (or acting as the interrupt thread when USE_INTERRUPT_THREAD is zero).
+	bool IsInInterruptThread() const;
+
+	// Called when the GPU has crashed. This function will not return.
+	void ProcessInterruptQueueOnGPUCrash();
+	void TerminateOnGPUCrash();
+
 protected:
 	void InitializeSubmissionPipe();
 	void ShutdownSubmissionPipe();
@@ -363,6 +372,10 @@ protected:
 	TMap<VkSemaphore, FBinarySemaphoreSignalInfo> SignaledSemaphores;
 	FGraphEventRef EopTask;
 
+	// Used when a dedicated interrupt thread is disabled (D3D12_USE_INTERRUPT_THREAD)
+	// Tracks the thread ID of the thread currently acting as the interrupt thread.
+	uint32 InterruptThreadID = 0;
+
 	uint32 ApiVersion = 0;
 	VkInstance Instance;
 	TArray<const ANSICHAR*> InstanceExtensions;
diff --git a/Engine/Source/Runtime/VulkanRHI/Private/VulkanExtensions.cpp b/Engine/Source/Runtime/VulkanRHI/Private/VulkanExtensions.cpp
--- a/Engine/Source/Runtime/VulkanRHI/Private/VulkanExtensions.cpp
+++ b/Engine/Source/Runtime/VulkanRHI/Private/VulkanExtensions.cpp
@@ -1235,18 +1235,12 @@ public:
 			DeviceDiagnosticsConfigCreateInfoNV.flags = VK_DEVICE_DIAGNOSTICS_CONFIG_ENABLE_SHADER_DEBUG_INFO_BIT_NV | VK_DEVICE_DIAGNOSTICS_CONFIG_ENABLE_RESOURCE_TRACKING_BIT_NV | VK_DEVICE_DIAGNOSTICS_CONFIG_ENABLE_AUTOMATIC_CHECKPOINTS_BIT_NV;
 			AddToPNext(DeviceCreateInfo, DeviceDiagnosticsConfigCreateInfoNV);
 
-#if NV_AFTERMATH && (RHI_NEW_GPU_PROFILER == 0)
-			// Vulkan's breadcrumb / markers implementation differs from the one provided by RHICore.
-			// @todo unify the implementation so this isn't necessary
-			extern void AftermathResolveMarkerCallback(const void* Marker, void** ResolvedMarkerData, uint32_t * MarkerSize);
-			UE::RHICore::Nvidia::Aftermath::InitializeBeforeDeviceCreation([](const void* MarkerData, const uint32_t MarkerDataSize, void* UserData, void** ResolvedMarkerData, uint32_t* ResolvedMarkerDataSize) -> void
-			{
-				AftermathResolveMarkerCallback(MarkerData, ResolvedMarkerData, ResolvedMarkerDataSize);
-			});
+#if NV_AFTERMATH
+			UE::RHICore::Nvidia::Aftermath::InitializeBeforeDeviceCreation(true);
 
 			extern void AftermathLateAssociate(float TimeLimitSeconds, uint32 FrameLimit);
 			UE::RHICore::Nvidia::Aftermath::SetLateShaderAssociateCallback(AftermathLateAssociate);
-#endif
+#endif // NV_AFTERMATH
 		}
 	}
 
diff --git a/Engine/Source/Runtime/VulkanRHI/Private/VulkanQueue.cpp b/Engine/Source/Runtime/VulkanRHI/Private/VulkanQueue.cpp
--- a/Engine/Source/Runtime/VulkanRHI/Private/VulkanQueue.cpp
+++ b/Engine/Source/Runtime/VulkanRHI/Private/VulkanQueue.cpp
@@ -41,13 +41,29 @@ FAutoConsoleVariableRef CVarVulkanSubmissionAllowTimelineSemaphores(
 //);
 
 
+static TAutoConsoleVariable<int32> CVarVulkanExtendedLifetimeFrames(
+	TEXT("r.Vulkan.Aftermath.ExtendedLifetimeFrames"),
+	2,
+	TEXT("Number of frames to keep breadcrumbs alive for Aftermath checkpoints."),
+	ECVF_ReadOnly
+);
+
+static TAutoConsoleVariable<int32> CVarVulkanDiagnosticBuffer(
+	TEXT("r.Vulkan.DiagnosticBuffer"),
+	1,
+	TEXT("0: Disable the diagnostic buffer\n")
+	TEXT("1: Enable the diagnostic buffer (default)"),
+	ECVF_ReadOnly
+);
+
+
 FVulkanQueue::FVulkanQueue(FVulkanDevice& InDevice, uint32 InFamilyIndex, EVulkanQueueType InQueueType)
 	: Queue(VK_NULL_HANDLE)
 	, FamilyIndex(InFamilyIndex)
 	, QueueIndex(0)
 	, QueueType(InQueueType)
 	, Device(InDevice)
-	, bUseTimelineSemaphores(InDevice.GetOptionalExtensions().HasKHRTimelineSemaphore && GAllowTimelineSemaphores)
+	, bUseTimelineSemaphores(InDevice.GetOptionalExtensions().HasKHRTimelineSemaphore&& GAllowTimelineSemaphores)
 {
 	VulkanRHI::vkGetDeviceQueue(Device.GetInstanceHandle(), FamilyIndex, QueueIndex, &Queue);
 
@@ -407,10 +423,12 @@ int32 FVulkanQueue::ProcessInterruptQueue(uint64 Timeout)
 		if (bSuccess)
 		{
 			++NumPayloads;
-			PendingInterrupt.Pop();
 
-			// Resolve any pending actions and delete the payload
+			// Resolve any pending actions from the payload being completed
 			FVulkanDynamicRHI::Get().CompletePayload(Payload);
+
+			PendingInterrupt.Pop();
+			delete Payload;
 		}
 
 	} while (bSuccess);
@@ -555,3 +573,121 @@ void FVulkanQueue::ReleaseCommandBufferPool(FVulkanCommandBufferPool* CommandBuf
 	TArray<FVulkanCommandBufferPool*>& CommandBufferPoolArray = CommandBufferPools[(int32)CommandBufferPool->GetCommandBufferType()];
 	CommandBufferPoolArray.Add(CommandBufferPool);
 }
+
+void FVulkanQueue::InitDiagnosticBuffer()
+{
+	check(!DiagnosticBuffer.IsValid());
+	if (CVarVulkanDiagnosticBuffer->GetInt())
+	{
+		DiagnosticBuffer = MakeUnique<FVulkanDiagnosticBuffer>(Device, *this);
+	}
+}
+
+
+
+
+
+
+
+FVulkanDiagnosticBuffer::FVulkanDiagnosticBuffer(FVulkanDevice& InDevice, FVulkanQueue& InQueue)
+	: Device(InDevice)
+	, Queue(InQueue)
+{
+	VkBufferCreateInfo CreateInfo;
+	ZeroVulkanStruct(CreateInfo, VK_STRUCTURE_TYPE_BUFFER_CREATE_INFO);
+	CreateInfo.size = SizeInBytes;
+	CreateInfo.usage = VK_BUFFER_USAGE_TRANSFER_DST_BIT;
+	VERIFYVULKANRESULT(VulkanRHI::vkCreateBuffer(Device.GetInstanceHandle(), &CreateInfo, VULKAN_CPU_ALLOCATOR, &Buffer));
+
+	const VulkanRHI::EVulkanAllocationFlags MemoryFlags =
+		VulkanRHI::EVulkanAllocationFlags::Dedicated |
+		VulkanRHI::EVulkanAllocationFlags::AutoBind |
+		VulkanRHI::EVulkanAllocationFlags::HostVisible |
+		VulkanRHI::EVulkanAllocationFlags::HostCached;
+	Device.GetMemoryManager().AllocateBufferMemory(Allocation, Buffer, MemoryFlags, TEXT("DiagnosticBuffer"));
+
+	Data = (FQueue*)Allocation.GetMappedPointer(&InDevice);
+	check(Data);
+	FMemory::Memzero(Data, SizeInBytes);
+
+#if WITH_RHI_BREADCRUMBS
+	if (Device.GetOptionalExtensions().HasNVDiagnosticCheckpoints &&
+		UE::RHI::UseGPUCrashBreadcrumbs())
+	{
+		ExtendedBreadcrumbAllocators.SetNumZeroed(CVarVulkanExtendedLifetimeFrames->GetInt());
+	}
+#endif
+}
+
+FVulkanDiagnosticBuffer::~FVulkanDiagnosticBuffer()
+{
+	if (Buffer != VK_NULL_HANDLE)
+	{
+		VulkanRHI::vkDestroyBuffer(Device.GetInstanceHandle(), Buffer, VULKAN_CPU_ALLOCATOR);
+		Buffer = VK_NULL_HANDLE;
+	}
+
+	if (Allocation.IsValid())
+	{
+		Device.GetMemoryManager().FreeVulkanAllocation(Allocation, VulkanRHI::EVulkanFreeFlag_DontDefer);
+	}
+}
+
+#if WITH_RHI_BREADCRUMBS
+void FVulkanDiagnosticBuffer::WriteMarkerIn(FVulkanCommandBuffer& CommandBuffer, FRHIBreadcrumbNode* Breadcrumb) const
+{
+	const VkCommandBuffer CommandBufferHandle = CommandBuffer.GetHandle();
+	const uint64 WriteOffset = offsetof(FQueue, MarkerIn);
+
+	if (Device.GetOptionalExtensions().HasAMDBufferMarker)
+	{
+		VulkanDynamicAPI::vkCmdWriteBufferMarkerAMD(CommandBufferHandle, VK_PIPELINE_STAGE_BOTTOM_OF_PIPE_BIT, Buffer, WriteOffset, Breadcrumb->ID);
+	}
+	else
+	{
+		// Buffer writes cannot be recorded inside of a render pass
+		if (CommandBuffer.IsOutsideRenderPass())
+		{
+			VulkanDynamicAPI::vkCmdPipelineBarrier(CommandBufferHandle, VK_PIPELINE_STAGE_BOTTOM_OF_PIPE_BIT, VK_PIPELINE_STAGE_TOP_OF_PIPE_BIT, 0, 0, nullptr, 0, nullptr, 0, nullptr);
+			VulkanDynamicAPI::vkCmdFillBuffer(CommandBufferHandle, Buffer, WriteOffset, sizeof(Breadcrumb->ID), Breadcrumb->ID);
+			VulkanDynamicAPI::vkCmdPipelineBarrier(CommandBufferHandle, VK_PIPELINE_STAGE_BOTTOM_OF_PIPE_BIT, VK_PIPELINE_STAGE_TOP_OF_PIPE_BIT, 0, 0, nullptr, 0, nullptr, 0, nullptr);
+		}
+	}
+
+#if NV_AFTERMATH
+	if (Device.GetOptionalExtensions().HasNVDiagnosticCheckpoints)
+	{
+		VulkanDynamicAPI::vkCmdSetCheckpointNV(CommandBufferHandle, Breadcrumb);
+	}
+#endif // NV_AFTERMATH
+}
+
+void FVulkanDiagnosticBuffer::WriteMarkerOut(FVulkanCommandBuffer& CommandBuffer, FRHIBreadcrumbNode* Breadcrumb) const
+{
+	const VkCommandBuffer CommandBufferHandle = CommandBuffer.GetHandle();
+	const uint64 WriteOffset = offsetof(FQueue, MarkerOut);
+
+	if (Device.GetOptionalExtensions().HasAMDBufferMarker)
+	{
+		VulkanDynamicAPI::vkCmdWriteBufferMarkerAMD(CommandBufferHandle, VK_PIPELINE_STAGE_BOTTOM_OF_PIPE_BIT, Buffer, WriteOffset, Breadcrumb->ID);
+	}
+	else
+	{
+		// Buffer writes cannot be recorded inside of a render pass
+		if (CommandBuffer.IsOutsideRenderPass())
+		{
+			VulkanDynamicAPI::vkCmdPipelineBarrier(CommandBufferHandle, VK_PIPELINE_STAGE_BOTTOM_OF_PIPE_BIT, VK_PIPELINE_STAGE_TOP_OF_PIPE_BIT, 0, 0, nullptr, 0, nullptr, 0, nullptr);
+			VulkanDynamicAPI::vkCmdFillBuffer(CommandBufferHandle, Buffer, WriteOffset, sizeof(Breadcrumb->ID), Breadcrumb->ID);
+			VulkanDynamicAPI::vkCmdPipelineBarrier(CommandBufferHandle, VK_PIPELINE_STAGE_BOTTOM_OF_PIPE_BIT, VK_PIPELINE_STAGE_TOP_OF_PIPE_BIT, 0, 0, nullptr, 0, nullptr, 0, nullptr);
+		}
+	}
+
+#if NV_AFTERMATH
+	if (Device.GetOptionalExtensions().HasNVDiagnosticCheckpoints)
+	{
+		VulkanDynamicAPI::vkCmdSetCheckpointNV(CommandBufferHandle, Breadcrumb->GetParent());
+	}
+#endif // NV_AFTERMATH
+}
+
+#endif // WITH_RHI_BREADCRUMBS
\ No newline at end of file
diff --git a/Engine/Source/Runtime/VulkanRHI/Private/VulkanQueue.h b/Engine/Source/Runtime/VulkanRHI/Private/VulkanQueue.h
--- a/Engine/Source/Runtime/VulkanRHI/Private/VulkanQueue.h
+++ b/Engine/Source/Runtime/VulkanRHI/Private/VulkanQueue.h
@@ -9,6 +9,8 @@
 #include "VulkanConfiguration.h"
 #include "VulkanBarriers.h"
 #include "VulkanCommandBuffer.h"
+#include "RHIDiagnosticBuffer.h"
+#include "RHICoreNvidiaAftermath.h"
 
 class FVulkanDynamicRHI;
 class FVulkanDevice;
@@ -30,6 +32,17 @@ enum class EVulkanQueueType : uint8
 	Count,
 };
 
+inline const TCHAR* GetVulkanQueueTypeName(EVulkanQueueType QueueType)
+{
+	switch (QueueType)
+	{
+	default: checkNoEntry(); // fallthrough
+	case EVulkanQueueType::Graphics:      return TEXT("Graphics");
+	case EVulkanQueueType::AsyncCompute:  return TEXT("AsyncCompute");
+	case EVulkanQueueType::Transfer:      return TEXT("Transfer");
+	}
+}
+
 struct FBinarySemaphoreSignalInfo
 {
 	uint64 TimelineValue;
@@ -101,7 +114,13 @@ public:
 #if RHI_NEW_GPU_PROFILER
 	UE::RHI::GPUProfiler::FQueue GetProfilerQueue() const;
 #endif
-	
+
+	void InitDiagnosticBuffer();
+	class FVulkanDiagnosticBuffer* GetDiagnosticBuffer()
+	{
+		return DiagnosticBuffer.Get();
+	}
+
 private:
 	void SubmitPayloads(TArrayView<FVulkanPayload*> Payloads, TMap<VkSemaphore, FBinarySemaphoreSignalInfo>& InOutSignaledSemas);
 	void Submit(TArrayView<VkSubmitInfo> InSubmitInfos, VulkanRHI::FFence* Fence);
@@ -141,5 +160,54 @@ private:
 	FVulkanTiming* Timing = nullptr;
 #endif
 
+	TUniquePtr<class FVulkanDiagnosticBuffer> DiagnosticBuffer;
+
 	void FillSupportedStageBits();
 };
+
+
+//
+// Per Queue diagnostic buffer. Stays accessible after a GPU crash to allow readback of diagnostic messages.
+// Also used to track the progress of the GPU via breadcrumb markers.
+//
+class FVulkanDiagnosticBuffer : public FRHIDiagnosticBuffer
+{
+public:
+	FVulkanDiagnosticBuffer(FVulkanDevice& InDevice, FVulkanQueue& InQueue);
+	~FVulkanDiagnosticBuffer();
+
+	// :todo: Bindings for shader asserts, etc
+
+#if WITH_RHI_BREADCRUMBS
+	void WriteMarkerIn(FVulkanCommandBuffer& CommandBuffer, FRHIBreadcrumbNode* Breadcrumb) const;
+	void WriteMarkerOut(FVulkanCommandBuffer& CommandBuffer, FRHIBreadcrumbNode* Breadcrumb) const;
+
+	uint32 ReadMarkerIn()
+	{
+		Allocation.InvalidateMappedMemory(&Device); 
+		return Data->MarkerIn; 
+	}
+	uint32 ReadMarkerOut()
+	{
+		Allocation.InvalidateMappedMemory(&Device);
+		return Data->MarkerOut; 
+	}
+#endif // WITH_RHI_BREADCRUMBS
+
+	bool IsValid() const { return (Buffer != VK_NULL_HANDLE); }
+
+private:
+	FVulkanDevice& Device;
+	FVulkanQueue& Queue;
+	VkBuffer Buffer = VK_NULL_HANDLE;
+	VulkanRHI::FVulkanAllocation Allocation;
+
+#if WITH_RHI_BREADCRUMBS
+	// Extend lifetime of breadcrumbs so we can use the pointers directly
+	// One array per frame, and cycle which array is used on EndFrame
+	TArray<FRHIBreadcrumbAllocatorArray> ExtendedBreadcrumbAllocators;
+	int32 CurrentExtendedIndex = 0;
+#endif // WITH_RHI_BREADCRUMBS
+
+	friend FVulkanDynamicRHI;
+};
\ No newline at end of file
diff --git a/Engine/Source/Runtime/VulkanRHI/Private/VulkanRHI.cpp b/Engine/Source/Runtime/VulkanRHI/Private/VulkanRHI.cpp
--- a/Engine/Source/Runtime/VulkanRHI/Private/VulkanRHI.cpp
+++ b/Engine/Source/Runtime/VulkanRHI/Private/VulkanRHI.cpp
@@ -1265,9 +1265,14 @@ void FVulkanCommandListContext::RHIEndDrawingViewport(FRHIViewport* ViewportRHI,
 #if WITH_RHI_BREADCRUMBS
 	void FVulkanCommandListContext::RHIBeginBreadcrumbGPU(FRHIBreadcrumbNode* Breadcrumb)
 	{
+		if (UE::RHI::UseGPUCrashBreadcrumbs() && Queue.GetDiagnosticBuffer())
+		{
+			Queue.GetDiagnosticBuffer()->WriteMarkerIn(GetCommandBuffer(), Breadcrumb);
+		}
+
 		const TCHAR* NameStr = nullptr;
 		FRHIBreadcrumb::FBuffer Buffer;
-		auto GetNameStr = [&]()
+		auto GetNameStr = [&NameStr, &Buffer, Breadcrumb]()
 		{
 			if (!NameStr)
 			{
@@ -1365,6 +1370,11 @@ void FVulkanCommandListContext::RHIEndDrawingViewport(FRHIViewport* ViewportRHI,
 			}
 		#endif
 		}
+
+		if (UE::RHI::UseGPUCrashBreadcrumbs() && Queue.GetDiagnosticBuffer())
+		{
+			Queue.GetDiagnosticBuffer()->WriteMarkerOut(GetCommandBuffer(), Breadcrumb);
+		}
 	}
 #endif // WITH_RHI_BREADCRUMBS
 
diff --git a/Engine/Source/Runtime/VulkanRHI/Private/VulkanSubmission.cpp b/Engine/Source/Runtime/VulkanRHI/Private/VulkanSubmission.cpp
--- a/Engine/Source/Runtime/VulkanRHI/Private/VulkanSubmission.cpp
+++ b/Engine/Source/Runtime/VulkanRHI/Private/VulkanSubmission.cpp
@@ -15,6 +15,7 @@
 #include "HAL/Runnable.h"
 #include "HAL/RunnableThread.h"
 
+extern CORE_API bool GIsGPUCrashed;
 
 FVulkanPayload::~FVulkanPayload()
 {
@@ -427,6 +428,12 @@ bool FVulkanDynamicRHI::ProcessInterruptQueue()
 	{
 		bProgress |= Queue.ProcessInterruptQueue(0) > 0;
 	});
+
+	if (GIsGPUCrashed)
+	{
+		TerminateOnGPUCrash();
+	}
+
 	return bProgress;
 }
 
@@ -449,6 +456,7 @@ void FVulkanDynamicRHI::ProcessInterruptQueueUntil(FGraphEvent* GraphEvent)
 		{
 			if (InterruptCS.TryLock())
 			{
+				TGuardValue<uint32> Guard(InterruptThreadID, FPlatformTLS::GetCurrentThreadId());
 				WaitAndProcessInterruptQueue();
 				InterruptCS.Unlock();
 			}
@@ -468,6 +476,7 @@ void FVulkanDynamicRHI::ProcessInterruptQueueUntil(FGraphEvent* GraphEvent)
 			bProgress = false;
 			if (InterruptCS.TryLock())
 			{
+				TGuardValue<uint32> Guard(InterruptThreadID, FPlatformTLS::GetCurrentThreadId());
 				bProgress = ProcessInterruptQueue();
 				InterruptCS.Unlock();
 			}
@@ -599,17 +608,39 @@ void FVulkanDynamicRHI::CompletePayload(FVulkanPayload* Payload)
 				TempResults.SetNumZeroed(QueryPool->GetMaxQueries());
 			}
 
-			const VkQueryResultFlags QueryFlags = VK_QUERY_RESULT_64_BIT | VK_QUERY_RESULT_WITH_AVAILABILITY_BIT | VK_QUERY_RESULT_WAIT_BIT;
-			const VkResult Result = VulkanRHI::vkGetQueryPoolResults(
-				Device->GetInstanceHandle(),
-				QueryPool->GetHandle(),
-				0,
-				QueryPool->CurrentQueryCount,
-				TempResults.Num() * TempResults.GetTypeSize(),
-				(uint64_t*)TempResults.GetData(),
-				sizeof(FQueryResult),
-				QueryFlags);
-			VERIFYVULKANRESULT_EXPANDED(Result);
+			// Using VK_QUERY_RESULT_WAIT_BIT can cause infinite waits on some versions of Adreno if the device is lost, 
+			// regardless of spec saying it'll eventually return.  So create a loop which shouldn't be too costly since
+			// the fence was reached guarateeing that all these queries are done anyhow.
+			VkResult Result = VK_SUCCESS;
+			while (Result >= VK_SUCCESS)
+			{
+				const VkQueryResultFlags QueryFlags = VK_QUERY_RESULT_64_BIT | VK_QUERY_RESULT_WITH_AVAILABILITY_BIT;
+				Result = VulkanRHI::vkGetQueryPoolResults(
+					Device->GetInstanceHandle(),
+					QueryPool->GetHandle(),
+					0,
+					QueryPool->CurrentQueryCount,
+					TempResults.Num() * TempResults.GetTypeSize(),
+					(uint64_t*)TempResults.GetData(),
+					sizeof(FQueryResult),
+					QueryFlags);
+
+				if (Result == VK_SUCCESS)
+				{
+					break;
+				}
+				else if (GIsGPUCrashed)
+				{
+					// The GPU is crashed and we are the current interrupt thread, 
+					// terminate right away if we crashed since no more progress is possible
+					TerminateOnGPUCrash();
+				}
+				else
+				{
+					FPlatformProcess::SleepNoStats(0.005f);
+					VERIFYVULKANRESULT_EXPANDED(Result);
+				}
+			}
 
 			for (uint32 QueryIndex = 0; QueryIndex < QueryPool->CurrentQueryCount; ++QueryIndex)
 			{
@@ -682,6 +713,37 @@ void FVulkanDynamicRHI::CompletePayload(FVulkanPayload* Payload)
 		QueryPoolArray.Reset();
 	}
 
+#if WITH_RHI_BREADCRUMBS
+	// For NV checkpoints, extend the lifetime of breadcrumbs by keeping an array of all the allocators
+	// populated until the next frame
+	if (Device->GetOptionalExtensions().HasNVDiagnosticCheckpoints &&
+		Payload->Queue.GetDiagnosticBuffer() && 
+		Payload->Queue.GetDiagnosticBuffer()->ExtendedBreadcrumbAllocators.Num())
+	{
+		FVulkanDiagnosticBuffer& DiagnosticBuffer = *Payload->Queue.GetDiagnosticBuffer();
+
+		if (Payload->BreadcrumbAllocators.IsValid())
+		{
+			for (TSharedRef<FRHIBreadcrumbAllocator>& Allocator : *Payload->BreadcrumbAllocators)
+			{
+				DiagnosticBuffer.ExtendedBreadcrumbAllocators[DiagnosticBuffer.CurrentExtendedIndex].AddUnique(&Allocator.Get());
+			}
+		}
+
+#if RHI_NEW_GPU_PROFILER
+		const bool bEndFrame = Payload->EndFrameEvent.IsSet();
+#else
+		const bool bEndFrame = Payload->bEndFrame;
+#endif
+		if (bEndFrame)
+		{
+			DiagnosticBuffer.CurrentExtendedIndex = (DiagnosticBuffer.CurrentExtendedIndex + 1) % DiagnosticBuffer.ExtendedBreadcrumbAllocators.Num();
+			DiagnosticBuffer.ExtendedBreadcrumbAllocators[DiagnosticBuffer.CurrentExtendedIndex].Reset();
+		}
+	}
+#endif // WITH_RHI_BREADCRUMBS
+
+
 #if RHI_NEW_GPU_PROFILER
 	if (!Payload->EventStream.IsEmpty())
 	{
@@ -724,8 +786,6 @@ void FVulkanDynamicRHI::CompletePayload(FVulkanPayload* Payload)
 		SyncPoint->DispatchSubsequents();
 	}
 	Payload->SyncPoints.Empty();
-
-	delete Payload;
 }
 
 void FVulkanPayload::PreExecute()
@@ -777,3 +837,47 @@ void FVulkanDynamicRHI::RHIRunOnQueue(EVulkanRHIRunOnQueueType QueueType, TFunct
 		SubmissionEvent->Wait();
 	}
 }
+
+void FVulkanDynamicRHI::ProcessInterruptQueueOnGPUCrash()
+{
+	// This function will not return.
+
+	if (InterruptThread)
+	{
+		// Since we have an interrupt thread, allow it to process the GPU crash.
+		// This is necessary so it can retrieve all the active payloads for resolving breadcrumbs.
+		InterruptThread->Kick();
+
+		// Wait for the interrupt thread to exit (which will never happen).
+		InterruptThread->Join();
+	}
+	else
+	{
+		// If we have no interrupt thread, assume ownership on the current thread
+		// (or block forever on the scope lock if multiple threads enter this function).
+		FScopeLock Lock(&InterruptCS);
+		TGuardValue<uint32> Guard(InterruptThreadID, FPlatformTLS::GetCurrentThreadId());
+
+		while (true)
+		{
+			ProcessInterruptQueue();
+		}
+	}
+}
+
+bool FVulkanDynamicRHI::IsInInterruptThread() const
+{
+	uint32 ThisThreadID = FPlatformTLS::GetCurrentThreadId();
+
+	// If we don't have a dedicated interrupt thread, the thread currently acting
+	// as the interrupt thread is tracked via the InterruptThreadID field.
+
+	if (InterruptThread)
+	{
+		return ThisThreadID == InterruptThread->GetThreadID();
+	}
+	else
+	{
+		return ThisThreadID == InterruptThreadID;
+	}
+}
\ No newline at end of file
diff --git a/Engine/Source/Runtime/VulkanRHI/Private/VulkanUtil.cpp b/Engine/Source/Runtime/VulkanRHI/Private/VulkanUtil.cpp
--- a/Engine/Source/Runtime/VulkanRHI/Private/VulkanUtil.cpp
+++ b/Engine/Source/Runtime/VulkanRHI/Private/VulkanUtil.cpp
@@ -16,6 +16,8 @@
 #include "PipelineStateCache.h"
 #include "Async/ParallelFor.h"
 #include "HAL/FileManager.h"
+#include "HAL/Runnable.h"
+#include "HAL/RunnableThread.h"
 #include "RenderCore.h"
 #include "RHICoreNvidiaAftermath.h"
 #include "VulkanRayTracing.h"
@@ -713,11 +715,151 @@
 		return Buffer;
 	}
 
-	void CheckDeviceFault(FVulkanDevice* InDevice)
+	/**
+	 * Checks that the given result isn't a failure.  If it is, the application exits with an appropriate error message.
+	 * @param	Result - The result code to check
+	 * @param	Code - The code which yielded the result.
+	 * @param	VkFunction - Tested function name.
+	 * @param	Filename - The filename of the source file containing Code.
+	 * @param	Line - The line number of Code within Filename.
+	 */
+	void VerifyVulkanResult(VkResult Result, const ANSICHAR* VkFunction, const ANSICHAR* Filename, uint32 Line)
 	{
-		if (InDevice->GetOptionalExtensions().HasEXTDeviceFault)
+		// It's possible for multiple threads to catch GPU crashes or other Vulkan errors at the same time. Make sure we only log the error once
+		// by acquiring this critical section inside VerifyVulkanResult (and never releasing it, because those functions don't return).
+		static FCriticalSection CrashCritSec;
+
+		bool bDumpMemory = false;
+		FString ErrorString;
+		switch (Result)
 		{
-			const VkDevice DeviceHandle = InDevice->GetInstanceHandle();
+#define VKERRORCASE(x)	case x: ErrorString = TEXT(#x)
+		VKERRORCASE(VK_NOT_READY); break;
+		VKERRORCASE(VK_TIMEOUT); break;
+		VKERRORCASE(VK_EVENT_SET); break;
+		VKERRORCASE(VK_EVENT_RESET); break;
+		VKERRORCASE(VK_INCOMPLETE); break;
+		VKERRORCASE(VK_ERROR_OUT_OF_HOST_MEMORY); bDumpMemory = true; break;
+		VKERRORCASE(VK_ERROR_OUT_OF_DEVICE_MEMORY); bDumpMemory = true; break;
+		VKERRORCASE(VK_ERROR_INITIALIZATION_FAILED); break;
+		VKERRORCASE(VK_ERROR_DEVICE_LOST); 
+			GIsGPUCrashed = true; 
+#if PLATFORM_ANDROID
+			FAndroidStats::LogGPUStats();
+#endif
+			break;
+		VKERRORCASE(VK_ERROR_MEMORY_MAP_FAILED); break;
+		VKERRORCASE(VK_ERROR_LAYER_NOT_PRESENT); break;
+		VKERRORCASE(VK_ERROR_EXTENSION_NOT_PRESENT); break;
+		VKERRORCASE(VK_ERROR_FEATURE_NOT_PRESENT); break;
+		VKERRORCASE(VK_ERROR_INCOMPATIBLE_DRIVER); break;
+		VKERRORCASE(VK_ERROR_TOO_MANY_OBJECTS); break;
+		VKERRORCASE(VK_ERROR_FORMAT_NOT_SUPPORTED); break;
+		VKERRORCASE(VK_ERROR_SURFACE_LOST_KHR); break;
+		VKERRORCASE(VK_ERROR_NATIVE_WINDOW_IN_USE_KHR); break;
+		VKERRORCASE(VK_SUBOPTIMAL_KHR); break;
+		VKERRORCASE(VK_ERROR_OUT_OF_DATE_KHR); break;
+		VKERRORCASE(VK_ERROR_INCOMPATIBLE_DISPLAY_KHR); break;
+		VKERRORCASE(VK_ERROR_VALIDATION_FAILED_EXT); break;
+		VKERRORCASE(VK_ERROR_INVALID_SHADER_NV); break;
+		VKERRORCASE(VK_ERROR_FRAGMENTED_POOL); break;
+		VKERRORCASE(VK_ERROR_OUT_OF_POOL_MEMORY_KHR); break;
+		VKERRORCASE(VK_ERROR_INVALID_EXTERNAL_HANDLE_KHR); break;
+		VKERRORCASE(VK_ERROR_NOT_PERMITTED_EXT); break;
+#undef VKERRORCASE
+		default:
+			ErrorString = VK_TYPE_TO_STRING(VkResult, Result);
+			break;
+		}
+
+#define VULKAN_PRINT_ERROR_MSG(Verbosity) UE_LOG(LogVulkanRHI, Error, TEXT("%s failed, VkResult=%d\n at %s:%u \n with error %s"), ANSI_TO_TCHAR(VkFunction), (int32)Result, ANSI_TO_TCHAR(Filename), Line, *ErrorString)
+
+		if (GIsGPUCrashed)
+		{
+			if (GVulkanRHI->IsInInterruptThread())
+			{
+				//
+				// We're already on the interrupt thread. We must not block on the lock, as the interrupt thread is responsible
+				// for reporting the GPU crash (we can only safely access the RHI breadcrumbs / active payloads from the interrupt thread).
+				//
+				// Attempt to take the crash lock to prevent other threads from log spamming, but don't give up if we fail to do so.
+				//
+				if (CrashCritSec.TryLock())
+				{
+					// We're the first thread to take the lock... log the error
+					VULKAN_PRINT_ERROR_MSG(Error);
+				}
+
+				// This function does not return.
+				GVulkanRHI->TerminateOnGPUCrash();
+			}
+			else
+			{
+				// Take the lock to ensure we report the error only once
+				CrashCritSec.Lock();
+				VULKAN_PRINT_ERROR_MSG(Error);
+
+				// The interrupt thread must be the one to handle VK_ERROR_DEVICE_LOST
+				// This function does not return.
+				GVulkanRHI->ProcessInterruptQueueOnGPUCrash();
+			}
+		}
+		else
+		{
+			CrashCritSec.Lock();
+
+			VULKAN_PRINT_ERROR_MSG(Error);
+
+#if UE_BUILD_DEBUG || UE_BUILD_DEVELOPMENT
+			if (bDumpMemory)
+			{
+				GVulkanRHI->DumpMemory();
+			}
+#endif
+		}
+
+		GLog->Panic();
+
+		VULKAN_PRINT_ERROR_MSG(Fatal);
+
+		// Force shutdown, we can't do anything useful anymore.
+		FPlatformMisc::RequestExit(true, TEXT("VerifyVulkanResult"));
+
+#undef VULKAN_PRINT_ERROR_MSG
+	}
+
+}
+
+void FVulkanDynamicRHI::DumpCrashMarkers()
+{
+	// Only one device in Vulkan
+	const uint32 DeviceIndex = 0;
+
+	// Log shader asserts / prints
+	{
+		FString ShaderDiagnostics;
+
+		Device->ForEachQueue([&ShaderDiagnostics](FVulkanQueue& Queue)
+			{
+				FVulkanDiagnosticBuffer* DiagnosticBuffer = Queue.GetDiagnosticBuffer();
+				if (DiagnosticBuffer)
+				{
+					const EVulkanQueueType QueueType = Queue.GetQueueType();
+					const TCHAR* QueueName = GetVulkanQueueTypeName(QueueType);
+					ShaderDiagnostics += DiagnosticBuffer->GetShaderDiagnosticMessages(DeviceIndex, (uint32)QueueType, QueueName);
+				}
+
+			});
+
+		if (!ShaderDiagnostics.IsEmpty())
+		{
+			UE_LOG(LogVulkanRHI, Error, TEXT("Shader diagnostic messages and asserts:%s\r\n"), *ShaderDiagnostics);
+		}
+	}
+
+	if (Device->GetOptionalExtensions().HasEXTDeviceFault)
+	{
+		const VkDevice DeviceHandle = Device->GetInstanceHandle();
 			VkResult Result;
 
 			VkDeviceFaultCountsEXT FaultCounts;
@@ -781,112 +923,136 @@
 				}
 			}
 		}
+
+
+#if WITH_RHI_BREADCRUMBS
+
+	// Log Nvidia checkpoints
+#if NV_AFTERMATH
+	if (Device->GetOptionalExtensions().HasNVDiagnosticCheckpoints)
+	{
+		struct FCheckpointDataNV : public VkCheckpointDataNV
+		{
+			FCheckpointDataNV()
+			{
+				ZeroVulkanStruct(*this, VK_STRUCTURE_TYPE_CHECKPOINT_DATA_NV);
+			}
+		};
+
+		Device->ForEachQueue([](FVulkanQueue& Queue)
+			{
+				TArray<FCheckpointDataNV> Checkpoints;
+				uint32 Num = 0;
+				const VkQueue QueueHandle = Queue.GetHandle();
+				VulkanDynamicAPI::vkGetQueueCheckpointDataNV(QueueHandle, &Num, nullptr);
+				if (Num > 0)
+				{
+					Checkpoints.AddDefaulted(Num);
+					VulkanDynamicAPI::vkGetQueueCheckpointDataNV(QueueHandle, &Num, &Checkpoints[0]);
+					check(Num == Checkpoints.Num());
+					for (uint32 Index = 0; Index < Num; ++Index)
+					{
+						check(Checkpoints[Index].sType == VK_STRUCTURE_TYPE_CHECKPOINT_DATA_NV);
+						FRHIBreadcrumbNode* Breadcrumb = (FRHIBreadcrumbNode*)Checkpoints[Index].pCheckpointMarker;
+						UE_LOG(LogVulkanRHI, Error, TEXT("[VK_NV_device_diagnostic_checkpoints] %i: Stage %s (0x%08x), %s"),
+							Index, VK_TYPE_TO_STRING(VkPipelineStageFlagBits, Checkpoints[Index].stage), Checkpoints[Index].stage,
+							*Breadcrumb->GetFullPath());
+					}
+					GLog->Panic();
+				}
+			});
+	}
+#endif // NV_AFTERMATH
+
+	// Log RHI breadcrumb data
+	if (UE::RHI::UseGPUCrashBreadcrumbs())
+	{
+		TMap<FRHIBreadcrumbState::FQueueID, TArray<FRHIBreadcrumbRange>> QueueRanges;
+		FRHIBreadcrumbState BreadcrumbState;
+
+		Device->ForEachQueue([&QueueRanges, &BreadcrumbState, DeviceIndex](FVulkanQueue& Queue)
+			{
+				FVulkanDiagnosticBuffer* DiagnosticBuffer = Queue.GetDiagnosticBuffer();
+				if (DiagnosticBuffer)
+				{
+					const EVulkanQueueType QueueType = Queue.GetQueueType();
+					const TCHAR* QueueName = GetVulkanQueueTypeName(QueueType);
+
+					ERHIPipeline Pipeline;
+					switch (QueueType)
+					{
+					default: return; // Skip pipelines that the RHI doesn't handle
+					case EVulkanQueueType::Graphics:     Pipeline = ERHIPipeline::Graphics; break;
+					case EVulkanQueueType::AsyncCompute: Pipeline = ERHIPipeline::AsyncCompute; break;
+					}
+
+					TArray<FRHIBreadcrumbRange>& Ranges = QueueRanges.Add({ DeviceIndex, Pipeline });
+
+					// Pull all incomplete payloads from the pending interrupt queue.
+					TArray<FVulkanPayload*> Payloads;
+					{
+						FVulkanPayload* Payload;
+						while (Queue.PendingInterrupt.Dequeue(Payload))
+						{
+							Payloads.Add(Payload);
+						}
+					}
+
+					// Extract the breadcrumb ranges for these payloads.
+					for (FVulkanPayload* Payload : Payloads)
+					{
+						if (Payload->BreadcrumbRange)
+						{
+							Ranges.AddUnique(Payload->BreadcrumbRange);
+						}
+					}
+
+					BreadcrumbState.Devices[DeviceIndex].Pipelines[Pipeline].MarkerOut = DiagnosticBuffer->ReadMarkerOut();
+					BreadcrumbState.Devices[DeviceIndex].Pipelines[Pipeline].MarkerIn = DiagnosticBuffer->ReadMarkerIn();
+				}
+			});
+
+		// Traverse the breadcrumb tree and log active GPU work
+		if (!QueueRanges.IsEmpty())
+		{
+			BreadcrumbState.DumpActiveBreadcrumbs(QueueRanges);
+		}
+	}
+#endif // WITH_RHI_BREADCRUMBS
+}
+
+
+
+
+void FVulkanDynamicRHI::TerminateOnGPUCrash()
+{
+	DumpCrashMarkers();
+
+	// Make sure we wait on the Aftermath crash dump before we crash.
+#if NV_AFTERMATH
+	TArray<UE::RHICore::Nvidia::Aftermath::FCrashResult> AftermathResults;
+	UE::RHICore::Nvidia::Aftermath::OnGPUCrash(AftermathResults);
+#endif
+
+	UGameEngine* GameEngine = Cast<UGameEngine>(GEngine);
+	if (GameEngine)
+	{
+		GameEngine->OnGPUCrash();
 	}
 
-	/**
-	 * Checks that the given result isn't a failure.  If it is, the application exits with an appropriate error message.
-	 * @param	Result - The result code to check
-	 * @param	Code - The code which yielded the result.
-	 * @param	VkFunction - Tested function name.
-	 * @param	Filename - The filename of the source file containing Code.
-	 * @param	Line - The line number of Code within Filename.
-	 */
-	void VerifyVulkanResult(VkResult Result, const ANSICHAR* VkFunction, const ANSICHAR* Filename, uint32 Line)
+	// hard break here when the debugger is attached
+	if (FPlatformMisc::IsDebuggerPresent())
 	{
-		bool bDumpMemory = false;
-		FString ErrorString;
-		switch (Result)
-		{
-#define VKERRORCASE(x)	case x: ErrorString = TEXT(#x)
-		VKERRORCASE(VK_NOT_READY); break;
-		VKERRORCASE(VK_TIMEOUT); break;
-		VKERRORCASE(VK_EVENT_SET); break;
-		VKERRORCASE(VK_EVENT_RESET); break;
-		VKERRORCASE(VK_INCOMPLETE); break;
-		VKERRORCASE(VK_ERROR_OUT_OF_HOST_MEMORY); bDumpMemory = true; break;
-		VKERRORCASE(VK_ERROR_OUT_OF_DEVICE_MEMORY); bDumpMemory = true; break;
-		VKERRORCASE(VK_ERROR_INITIALIZATION_FAILED); break;
-		VKERRORCASE(VK_ERROR_DEVICE_LOST); 
-			GIsGPUCrashed = true; 
-#if PLATFORM_ANDROID
-			FAndroidStats::LogGPUStats();
-#endif
-			break;
-		VKERRORCASE(VK_ERROR_MEMORY_MAP_FAILED); break;
-		VKERRORCASE(VK_ERROR_LAYER_NOT_PRESENT); break;
-		VKERRORCASE(VK_ERROR_EXTENSION_NOT_PRESENT); break;
-		VKERRORCASE(VK_ERROR_FEATURE_NOT_PRESENT); break;
-		VKERRORCASE(VK_ERROR_INCOMPATIBLE_DRIVER); break;
-		VKERRORCASE(VK_ERROR_TOO_MANY_OBJECTS); break;
-		VKERRORCASE(VK_ERROR_FORMAT_NOT_SUPPORTED); break;
-		VKERRORCASE(VK_ERROR_SURFACE_LOST_KHR); break;
-		VKERRORCASE(VK_ERROR_NATIVE_WINDOW_IN_USE_KHR); break;
-		VKERRORCASE(VK_SUBOPTIMAL_KHR); break;
-		VKERRORCASE(VK_ERROR_OUT_OF_DATE_KHR); break;
-		VKERRORCASE(VK_ERROR_INCOMPATIBLE_DISPLAY_KHR); break;
-		VKERRORCASE(VK_ERROR_VALIDATION_FAILED_EXT); break;
-		VKERRORCASE(VK_ERROR_INVALID_SHADER_NV); break;
-		VKERRORCASE(VK_ERROR_FRAGMENTED_POOL); break;
-		VKERRORCASE(VK_ERROR_OUT_OF_POOL_MEMORY_KHR); break;
-		VKERRORCASE(VK_ERROR_INVALID_EXTERNAL_HANDLE_KHR); break;
-		VKERRORCASE(VK_ERROR_NOT_PERMITTED_EXT); break;
-#undef VKERRORCASE
-		default:
-			ErrorString = VK_TYPE_TO_STRING(VkResult, Result);
-			break;
-		}
+		UE_DEBUG_BREAK();
+	}
 
-#if VULKAN_HAS_DEBUGGING_ENABLED
-		if (Result == VK_ERROR_VALIDATION_FAILED_EXT)
-		{
-			if (GValidationCvar.GetValueOnRenderThread() == 0)
-			{
-				UE_LOG(LogVulkanRHI, Fatal, TEXT("Failed with Validation error. Try running with r.Vulkan.EnableValidation=1 or -vulkandebug to get information from the validation layers."));
-			}
-		}
-#endif
-
-		UE_LOG(LogVulkanRHI, Error, TEXT("%s failed, VkResult=%d\n at %s:%u \n with error %s"),
-			ANSI_TO_TCHAR(VkFunction), (int32)Result, ANSI_TO_TCHAR(Filename), Line, *ErrorString);
-
-		if (GIsGPUCrashed)
-		{
-			FVulkanDevice* Device = GVulkanRHI->GetDevice();
-
-#if RHI_NEW_GPU_PROFILER
-			checkNoEntry(); // @todo - new gpu profiler
-#endif
-
-			CheckDeviceFault(Device);
-
-			// Make sure we wait on the Aftermath crash dump before we crash.
-#if NV_AFTERMATH
-			TArray<UE::RHICore::Nvidia::Aftermath::FCrashResult> AftermathResults;
-			UE::RHICore::Nvidia::Aftermath::OnGPUCrash(AftermathResults);
-#endif
-
-			UGameEngine* GameEngine = Cast<UGameEngine>(GEngine);
-			if (GameEngine)
-			{
-				GameEngine->OnGPUCrash();
-			}
-		}
-
-#if UE_BUILD_DEBUG || UE_BUILD_DEVELOPMENT
-		if (bDumpMemory)
-		{
-			GVulkanRHI->DumpMemory();
-		}
-#endif
-		UE_LOG(LogVulkanRHI, Fatal, TEXT("%s failed, VkResult=%d\n at %s:%u \n with error %s"),
-			ANSI_TO_TCHAR(VkFunction), (int32)Result, ANSI_TO_TCHAR(Filename), Line, *ErrorString);
-
 		// Force shutdown, we can't do anything useful anymore.
-//		FPlatformMisc::RequestExit(true, TEXT("VerifyVulkanResult"));
-	}
+	FPlatformMisc::RequestExit(true, TEXT("FVulkanDynamicRHI.TerminateOnGPUCrash"));
 }
 
 
+
+
 DEFINE_STAT(STAT_VulkanNumPSOs);
 DEFINE_STAT(STAT_VulkanNumGraphicsPSOs);
 DEFINE_STAT(STAT_VulkanNumPSOLRU);
diff --git a/Engine/Source/Runtime/VulkanRHI/Private/VulkanViewport.cpp b/Engine/Source/Runtime/VulkanRHI/Private/VulkanViewport.cpp
--- a/Engine/Source/Runtime/VulkanRHI/Private/VulkanViewport.cpp
+++ b/Engine/Source/Runtime/VulkanRHI/Private/VulkanViewport.cpp
@@ -899,7 +899,6 @@ void FVulkanViewport::DestroySwapchain(FVulkanSwapChainRecreateInfo* RecreateInf
 
 inline static void CopyImageToBackBuffer(FVulkanCommandListContext& Context, FVulkanTexture& SrcSurface, FVulkanTexture& DstSurface, int32 SizeX, int32 SizeY, int32 WindowSizeX, int32 WindowSizeY, VkSurfaceTransformFlagBitsKHR CachedSurfaceTransform)
 {
-	RHI_BREADCRUMB_EVENT(Context, "CopyImageToBackBuffer");
 	const bool bNeedsVulkanPreTransform = CachedSurfaceTransform != VK_SURFACE_TRANSFORM_IDENTITY_BIT_KHR;
 
 	FVulkanCommandBuffer* CmdBuffer = &Context.GetCommandBuffer();
