VulkanRHI: Add descriptor cache support to transient allocator.  Allow transient allocator to be enabled on mobile, but keep disabled by default for now. (WIP)

#rnx

diff --git a/Engine/Config/Android/AndroidEngine.ini b/Engine/Config/Android/AndroidEngine.ini
--- a/Engine/Config/Android/AndroidEngine.ini
+++ b/Engine/Config/Android/AndroidEngine.ini
@@ -98,6 +98,8 @@ r.Vulkan.AllowSplitBarriers=0
 r.LensFlareBlurComputeShader=1
 ; Crash reporting diagnotic buffer
 r.Vulkan.DiagnosticBuffer=1
+; Transient memory allocator
+r.Vulkan.EnableTransientResourceAllocator=0
 
 [/Script/Engine.RendererSettings]
 r.DistanceFields=1
diff --git a/Engine/Source/Runtime/VulkanRHI/Private/VulkanRHI.cpp b/Engine/Source/Runtime/VulkanRHI/Private/VulkanRHI.cpp
--- a/Engine/Source/Runtime/VulkanRHI/Private/VulkanRHI.cpp
+++ b/Engine/Source/Runtime/VulkanRHI/Private/VulkanRHI.cpp
@@ -2112,9 +2112,7 @@ void FVulkanDynamicRHI::VulkanSetImageLayout(VkCommandBuffer CmdBuffer, VkImage
 IRHITransientResourceAllocator* FVulkanDynamicRHI::RHICreateTransientResourceAllocator()
 {
 #if VULKAN_SUPPORTS_TRANSIENT_RESOURCE_ALLOCATOR
-	// Only use transient heap on desktop platforms for now
-	// Not compatible with VulkanDescriptorCache for now because it hashes using the 32bit BufferId instead of the VulkanHandle.
-	if (GVulkanEnableTransientResourceAllocator && IsPCPlatform(GMaxRHIShaderPlatform) && !UseVulkanDescriptorCache())
+	if (GVulkanEnableTransientResourceAllocator)
 	{
 		return new FVulkanTransientResourceAllocator(Device->GetOrCreateTransientHeapCache());
 	}
diff --git a/Engine/Source/Runtime/VulkanRHI/Private/VulkanTransientResourceAllocator.cpp b/Engine/Source/Runtime/VulkanRHI/Private/VulkanTransientResourceAllocator.cpp
--- a/Engine/Source/Runtime/VulkanRHI/Private/VulkanTransientResourceAllocator.cpp
+++ b/Engine/Source/Runtime/VulkanRHI/Private/VulkanTransientResourceAllocator.cpp
@@ -10,7 +10,6 @@
 FVulkanTransientHeap::FVulkanTransientHeap(const FInitializer& Initializer, FVulkanDevice* InDevice)
 	: FRHITransientHeap(Initializer)
 	, FDeviceChild(InDevice)
-	, VulkanBuffer(VK_NULL_HANDLE)
 {
 	EBufferUsageFlags UEBufferUsageFlags = BUF_UniformBuffer | BUF_VertexBuffer | BUF_IndexBuffer | BUF_DrawIndirect 
 		| BUF_UnorderedAccess | BUF_StructuredBuffer | BUF_ShaderResource | BUF_KeepCPUAccessible;
@@ -25,6 +24,10 @@ FVulkanTransientHeap::FVulkanTransientHeap(const FInitializer& Initializer, FVul
 	const bool bZeroSize = false;
 	const VkBufferUsageFlags BufferUsageFlags = FVulkanBuffer::UEToVKBufferUsageFlags(InDevice, UEBufferUsageFlags, bZeroSize);
 	VulkanBuffer = InDevice->CreateBuffer(Initializer.Size, BufferUsageFlags);
+	if (UseVulkanDescriptorCache())
+	{
+		HandleID = ++GVulkanBufferHandleIdCounter;
+	}
 
 	// Find the alignment that works for everyone
 	const uint32 MinBufferAlignment = FMath::Max<uint32>(Initializer.Alignment, VulkanRHI::FMemoryManager::CalculateBufferAlignment(*InDevice, UEBufferUsageFlags, bZeroSize));
@@ -38,6 +41,7 @@ FVulkanTransientHeap::~FVulkanTransientHeap()
 	Device->GetMemoryManager().FreeVulkanAllocation(InternalAllocation);
 	Device->GetDeferredDeletionQueue().EnqueueResource(VulkanRHI::FDeferredDeletionQueue2::EType::Buffer, VulkanBuffer);
 	VulkanBuffer = VK_NULL_HANDLE;
+	HandleID = 0;
 }
 
 VkDeviceMemory FVulkanTransientHeap::GetMemoryHandle()
@@ -53,6 +57,7 @@ VulkanRHI::FVulkanAllocation FVulkanTransientHeap::GetVulkanAllocation(const FRH
 	VulkanRHI::FVulkanAllocation TransientAlloc;
 	TransientAlloc.Reference(Heap->InternalAllocation);
 	TransientAlloc.VulkanHandle = (uint64)Heap->VulkanBuffer;
+	TransientAlloc.HandleId = Heap->HandleID;
 	TransientAlloc.Offset += HeapAllocation.Offset;
 	TransientAlloc.Size = HeapAllocation.Size;
 	check((TransientAlloc.Offset + TransientAlloc.Size) <= Heap->InternalAllocation.Size);
diff --git a/Engine/Source/Runtime/VulkanRHI/Private/VulkanTransientResourceAllocator.h b/Engine/Source/Runtime/VulkanRHI/Private/VulkanTransientResourceAllocator.h
--- a/Engine/Source/Runtime/VulkanRHI/Private/VulkanTransientResourceAllocator.h
+++ b/Engine/Source/Runtime/VulkanRHI/Private/VulkanTransientResourceAllocator.h
@@ -19,7 +19,8 @@ public:
 	static VulkanRHI::FVulkanAllocation GetVulkanAllocation(const FRHITransientHeapAllocation& TransientInitializer);
 
 private:
-	VkBuffer VulkanBuffer;
+	VkBuffer VulkanBuffer = VK_NULL_HANDLE;
+	uint32 HandleID = 0;
 	VkMemoryRequirements MemoryRequirements;
 	VulkanRHI::FVulkanAllocation InternalAllocation;
 };
