VulkanRHI: Fix descriptor set leak in parallel contexts (desktop SM5).

#rnx

diff --git a/Engine/Source/Runtime/VulkanRHI/Private/VulkanContext.cpp b/Engine/Source/Runtime/VulkanRHI/Private/VulkanContext.cpp
--- a/Engine/Source/Runtime/VulkanRHI/Private/VulkanContext.cpp
+++ b/Engine/Source/Runtime/VulkanRHI/Private/VulkanContext.cpp
@@ -280,6 +280,8 @@ void FVulkanCommandListContext::RHIEndParallelRenderPass()
 				ParentPayload.QueryPools[QueryPoolTypeIndex].Append(MoveTemp(Payload->QueryPools[QueryPoolTypeIndex]));
 			}
 
+			ParentPayload.DescriptorPoolSetContainers.Append(MoveTemp(Payload->DescriptorPoolSetContainers));
+
 			for (FVulkanCommandBuffer* SecondaryCommandBuffer : Payload->CommandBuffers)
 			{
 				CommandBufferHandles.Add(SecondaryCommandBuffer->GetHandle());
