VulkanRHI: set proper state for secondary command buffers (always inside renderpass)

#rnx

diff --git a/Engine/Source/Runtime/VulkanRHI/Private/VulkanCommandBuffer.cpp b/Engine/Source/Runtime/VulkanRHI/Private/VulkanCommandBuffer.cpp
--- a/Engine/Source/Runtime/VulkanRHI/Private/VulkanCommandBuffer.cpp
+++ b/Engine/Source/Runtime/VulkanRHI/Private/VulkanCommandBuffer.cpp
@@ -112,7 +112,8 @@ void FVulkanCommandBuffer::BeginRenderPass(const FVulkanBeginRenderPassInfo& Beg
 
 void FVulkanCommandBuffer::End(FVulkanQueryPool* QueryPool)
 {
-	checkf(IsOutsideRenderPass(), TEXT("Can't End as we're inside a render pass! CmdBuffer 0x%p State=%d"), CommandBufferHandle, (int32)State);
+	checkf(IsOutsideRenderPass() || (GetCommandBufferType() == EVulkanCommandBufferType::Secondary), 
+		TEXT("Can't End as we're inside a render pass! CmdBuffer 0x%p State=%d"), CommandBufferHandle, (int32)State);
 
 	// Reset barrier events for next use
 	for (VkEvent BarrierEvent : EndedBarrierEvents)
@@ -154,7 +155,7 @@ void FVulkanCommandBuffer::Begin(FVulkanQueryPool* QueryPool, VkRenderPass Rende
 		{
 			checkf(State == EState::ReadyForBegin, TEXT("Can't Begin as we're NOT ready! CmdBuffer 0x%p State=%d"), CommandBufferHandle, (int32)State);
 		}
-		State = EState::IsInsideBegin;
+		State = (RenderPassHandle != VK_NULL_HANDLE) ? EState::IsInsideRenderPass : EState::IsInsideBegin;
 	}
 
 	VkCommandBufferBeginInfo BeginInfo;
