VulkanRHI: Add more information in swapchain error message

#rnx

diff --git a/Engine/Source/Runtime/VulkanRHI/Private/VulkanViewport.cpp b/Engine/Source/Runtime/VulkanRHI/Private/VulkanViewport.cpp
--- a/Engine/Source/Runtime/VulkanRHI/Private/VulkanViewport.cpp
+++ b/Engine/Source/Runtime/VulkanRHI/Private/VulkanViewport.cpp
@@ -791,8 +791,10 @@ void FVulkanViewport::CreateSwapchain(FVulkanCommandListContext& Context, FVulka
 
 			checkf(Images.Num() >= NUM_BUFFERS, TEXT("We wanted at least %i images, actual Num: %i"), NUM_BUFFERS, Images.Num());
 
-			bool bCreateSemaphores = RenderingDoneSemaphores.IsEmpty();
-			checkf(bCreateSemaphores || RenderingDoneSemaphores.Num() == Images.Num(), TEXT("CreateSwapchain, image count is not expected to change"));
+			const bool bCreateSemaphores = RenderingDoneSemaphores.IsEmpty();
+			checkf(bCreateSemaphores || RenderingDoneSemaphores.Num() == Images.Num(), 
+				TEXT("CreateSwapchain, image count is not expected to change:  Requested=%d  Actual=%d  Previous=%d"),
+				NUM_BUFFERS, Images.Num(), RenderingDoneSemaphores.Num());
 
 			BackBufferImages.SetNum(Images.Num());
 			RenderingDoneSemaphores.SetNum(Images.Num());
