Update FVulkanVertexDeclarationKey to sort vertex elements before hashing, matching Pipeline File Cache behavior and preventing cache misses due to being directed to a different chunk in the Vulkan Chunked Pipeline Cache.
Also update FVulkanDynamicRHI::RHIComputeStatePrecachePSOHash to ignore bAllowVariableRateShading (forcing to true), since bAllowVariableRateShading is not reflected in the Pipeline File Cache, which led to similar cache misses for PSOs with bAllowVariableRateShading=false. Pipelines that vary only by bAllowVariableRateShading should be stored in the same chunk, since some drivers only need to compile once for both PSOs.

diff --git a/Engine/Source/Runtime/VulkanRHI/Private/VulkanRHI.cpp b/Engine/Source/Runtime/VulkanRHI/Private/VulkanRHI.cpp
--- a/Engine/Source/Runtime/VulkanRHI/Private/VulkanRHI.cpp
+++ b/Engine/Source/Runtime/VulkanRHI/Private/VulkanRHI.cpp
@@ -2122,7 +2122,7 @@ IRHITransientResourceAllocator* FVulkanDynamicRHI::RHICreateTransientResourceAll
 
 uint32 FVulkanDynamicRHI::GetPrecachePSOHashVersion()
 {
-	static const uint32 PrecacheHashVersion = 3;
+	static const uint32 PrecacheHashVersion = 4;
 	return PrecacheHashVersion;
 }
 
@@ -2194,7 +2194,11 @@ uint64 FVulkanDynamicRHI::RHIComputeStatePrecachePSOHash(const FGraphicsPipeline
 	HashKey.DrawShadingRate = Initializer.ShadingRate;
 	HashKey.PrimitiveType = Initializer.PrimitiveType;
 	HashKey.bDepthBounds = Initializer.bDepthBounds;
-	HashKey.bAllowVariableRateShading = Initializer.bAllowVariableRateShading;
+
+	// For hotfix 5.6.1, bAllowVariableRateShading is ignored for precaching hashes because it is not reflected in Pipeline File Cache entries
+	// This previously led cache reads for PSOs with bAllowVariableRateShading=false to be directed to the wrong chunk in the Vulkan Chunked Pipeline Cache and miss
+	// This fix will still result in a miss for drivers where toggling bAllowVariableRateShading requires a PSO recompile
+	HashKey.bAllowVariableRateShading = true;
 
 	uint64 PrecachePSOHash = CityHash64((const char*)&HashKey, sizeof(FHashKey));
 
diff --git a/Engine/Source/Runtime/VulkanRHI/Private/VulkanVertexDeclaration.cpp b/Engine/Source/Runtime/VulkanRHI/Private/VulkanVertexDeclaration.cpp
--- a/Engine/Source/Runtime/VulkanRHI/Private/VulkanVertexDeclaration.cpp
+++ b/Engine/Source/Runtime/VulkanRHI/Private/VulkanVertexDeclaration.cpp
@@ -18,10 +18,40 @@ struct FVulkanVertexDeclarationKey
 	{
 		Hash = FCrc::MemCrc_DEPRECATED(VertexElements.GetData(), VertexElements.Num()*sizeof(FVertexElement));
 
+		FVertexDeclarationElementList SortedElements = VertexElements;
+		SortedElements.Sort([](FVertexElement const& A, FVertexElement const& B)
+		{
+			if (A.StreamIndex < B.StreamIndex)
+			{
+				return true;
+			}
+			if (A.StreamIndex > B.StreamIndex)
+			{
+				return false;
+			}
+			if (A.Offset < B.Offset)
+			{
+				return true;
+			}
+			if (A.Offset > B.Offset)
+			{
+				return false;
+			}
+			if (A.AttributeIndex < B.AttributeIndex)
+			{
+				return true;
+			}
+			if (A.AttributeIndex > B.AttributeIndex)
+			{
+				return false;
+			}
+			return false;
+		});
+
 		HashNoStride = 0;
-		for (int32 ElementIndex = 0; ElementIndex < VertexElements.Num(); ElementIndex++)
+		for (int32 ElementIndex = 0; ElementIndex < SortedElements.Num(); ElementIndex++)
 		{
-			FVertexElement Element = VertexElements[ElementIndex];
+			FVertexElement Element = SortedElements[ElementIndex];
 			Element.Stride = 0;
 			HashNoStride = FCrc::MemCrc32(&Element, sizeof(Element), HashNoStride);
 		}
