VulkanRHI: Move internally created view ownership from the FVulkanFramebuffer to the FVulkanTexture itself (for upcoming dynamic rendering).  Create a FVulkanTextureViewDesc struct to package all the parameters used for internal vew creation.  Move the view creation logic for frame buffers into externally accessible functions (for upcoming dynamic rendering).

#rnx

diff --git a/Engine/Source/Runtime/VulkanRHI/Private/VulkanRHIPrivate.h b/Engine/Source/Runtime/VulkanRHI/Private/VulkanRHIPrivate.h
--- a/Engine/Source/Runtime/VulkanRHI/Private/VulkanRHIPrivate.h
+++ b/Engine/Source/Runtime/VulkanRHI/Private/VulkanRHIPrivate.h
@@ -274,7 +274,6 @@ public:
 		return PartialDepthTextureView->GetTextureView();
 	}
 
-	TIndirectArray<FVulkanView> OwnedTextureViews;
 	TArray<FVulkanView const*> AttachmentTextureViews;
 
 	// Copy from the Depth render target partial view
@@ -306,6 +305,13 @@ public:
 		return RenderArea;
 	}
 
+	// Expose the view creation logic so it can be reused for dynamic rendering
+	static FVulkanView* GetColorRenderTargetViewDesc(FVulkanTexture* Texture, uint32 MipIndex, int32 ArraySliceIndex, uint32 MultiViewCount, uint32& InOutNumLayers);
+	static FVulkanView* GetColorResolveTargetViewDesc(FVulkanTexture* ResolveTexture, uint32 MipIndex, int32 ArraySliceIndex);
+	static FVulkanView* GetDepthStencilTargetViewDesc(FVulkanTexture* Texture, uint32 NumColorAttachments, uint32 MipIndex, uint32& InOutNumLayers);
+	static FVulkanView* GetDepthStencilResolveTargetViewDesc(FVulkanTexture* ResolveTexture, uint32 MipIndex);
+	static FVulkanView* GetFragmentDensityAttachmentViewDesc(FVulkanTexture* Texture, uint32 MipIndex);
+
 private:
 	VkFramebuffer Framebuffer;
 	VkRect2D RenderArea;
diff --git a/Engine/Source/Runtime/VulkanRHI/Private/VulkanRenderTarget.cpp b/Engine/Source/Runtime/VulkanRHI/Private/VulkanRenderTarget.cpp
--- a/Engine/Source/Runtime/VulkanRHI/Private/VulkanRenderTarget.cpp
+++ b/Engine/Source/Runtime/VulkanRHI/Private/VulkanRenderTarget.cpp
@@ -572,7 +572,7 @@ void FVulkanCommandListContext::RHIBeginRenderPass(const FRHIRenderPassInfo& InI
 	FRHISetRenderTargetsInfo RTInfo;
 	InInfo.ConvertToRenderTargetsInfo(RTInfo);
 
-	FVulkanFramebuffer* Framebuffer = Device.GetRenderPassManager().GetOrCreateFramebuffer(RTInfo, RTLayout, RenderPass);
+	FVulkanFramebuffer* Framebuffer = Device.GetRenderPassManager().GetOrCreateFramebuffer(RTInfo, RTLayout, *RenderPass);
 	checkf(RenderPass != nullptr && Framebuffer != nullptr, TEXT("RenderPass not started! Bad combination of values? Depth %p #Color %d Color0 %p"), (void*)InInfo.DepthStencilRenderTarget.DepthStencilTarget, InInfo.GetNumColorRenderTargets(), (void*)InInfo.ColorRenderTargets[0].RenderTarget);
 
 	const bool bIsParallelRenderPass = (CurrentParallelRenderPassInfo != nullptr);
diff --git a/Engine/Source/Runtime/VulkanRHI/Private/VulkanRenderpass.cpp b/Engine/Source/Runtime/VulkanRHI/Private/VulkanRenderpass.cpp
--- a/Engine/Source/Runtime/VulkanRHI/Private/VulkanRenderpass.cpp
+++ b/Engine/Source/Runtime/VulkanRHI/Private/VulkanRenderpass.cpp
@@ -52,10 +52,8 @@ FVulkanRenderPassManager::~FVulkanRenderPassManager()
 	Framebuffers.Reset();
 }
 
-FVulkanFramebuffer* FVulkanRenderPassManager::GetOrCreateFramebuffer(const FRHISetRenderTargetsInfo& RenderTargetsInfo, const FVulkanRenderTargetLayout& RTLayout, FVulkanRenderPass* RenderPass)
+FVulkanFramebuffer* FVulkanRenderPassManager::GetOrCreateFramebuffer(const FRHISetRenderTargetsInfo& RenderTargetsInfo, const FVulkanRenderTargetLayout& RTLayout, FVulkanRenderPass& RenderPass)
 {
-	// todo-jn: threadsafe?
-
 	uint32 RTLayoutHash = RTLayout.GetRenderPassCompatibleHash();
 
 	uint64 MipsAndSlicesValues[MaxSimultaneousRenderTargets];
@@ -121,7 +119,7 @@ FVulkanFramebuffer* FVulkanRenderPassManager::GetOrCreateFramebuffer(const FRHIS
 		}
 	}
 
-	FVulkanFramebuffer* Framebuffer = new FVulkanFramebuffer(Device, RenderTargetsInfo, RTLayout, *RenderPass);
+	FVulkanFramebuffer* Framebuffer = new FVulkanFramebuffer(Device, RenderTargetsInfo, RTLayout, RenderPass);
 	FramebufferList->Framebuffer.Add(Framebuffer);
 	return Framebuffer;
 }
diff --git a/Engine/Source/Runtime/VulkanRHI/Private/VulkanRenderpass.h b/Engine/Source/Runtime/VulkanRHI/Private/VulkanRenderpass.h
--- a/Engine/Source/Runtime/VulkanRHI/Private/VulkanRenderpass.h
+++ b/Engine/Source/Runtime/VulkanRHI/Private/VulkanRenderpass.h
@@ -865,7 +865,7 @@ public:
 	FVulkanRenderPassManager(FVulkanDevice& InDevice) : Device(InDevice) {}
 	~FVulkanRenderPassManager();
 
-	FVulkanFramebuffer* GetOrCreateFramebuffer(const FRHISetRenderTargetsInfo& RenderTargetsInfo, const FVulkanRenderTargetLayout& RTLayout, FVulkanRenderPass* RenderPass);
+	FVulkanFramebuffer* GetOrCreateFramebuffer(const FRHISetRenderTargetsInfo& RenderTargetsInfo, const FVulkanRenderTargetLayout& RTLayout, FVulkanRenderPass& RenderPass);
 	FVulkanRenderPass* GetOrCreateRenderPass(const FVulkanRenderTargetLayout& RTLayout)
 	{
 		const uint32 RenderPassHash = RTLayout.GetRenderPassFullHash();
diff --git a/Engine/Source/Runtime/VulkanRHI/Private/VulkanResources.h b/Engine/Source/Runtime/VulkanRHI/Private/VulkanResources.h
--- a/Engine/Source/Runtime/VulkanRHI/Private/VulkanResources.h
+++ b/Engine/Source/Runtime/VulkanRHI/Private/VulkanResources.h
@@ -405,6 +405,24 @@ struct FVulkanCpuReadbackBuffer
 	uint32 MipOffsets[MAX_TEXTURE_MIP_COUNT];
 };
 
+
+struct FVulkanTextureViewDesc
+{
+	VkImageViewType ViewType = VK_IMAGE_VIEW_TYPE_MAX_ENUM;
+	VkImageAspectFlags AspectFlags = VK_IMAGE_ASPECT_NONE;
+	EPixelFormat UEFormat = PF_Unknown;
+	VkFormat Format = VK_FORMAT_UNDEFINED;
+	uint32 FirstMip = 0;
+	uint32 NumMips = VK_REMAINING_MIP_LEVELS;
+	uint32 ArraySliceIndex = 0;
+	uint32 NumArraySlices = VK_REMAINING_ARRAY_LAYERS;
+	VkImageUsageFlags ImageUsageFlags = 0;
+	VkSamplerYcbcrConversion SamplerYcbcrConversion = VK_NULL_HANDLE;
+	bool bUseIdentitySwizzle = false;
+
+	uint32 GetHash() const;
+};
+
 class FVulkanView
 {
 public:
@@ -486,19 +504,7 @@ public:
 		, uint32 InOffset
 		, uint32 InSize);
 
-	FVulkanView* InitAsTextureView(
-		  VkImage InImage
-		, VkImageViewType ViewType
-		, VkImageAspectFlags AspectFlags
-		, EPixelFormat UEFormat
-		, VkFormat Format
-		, uint32 FirstMip
-		, uint32 NumMips
-		, uint32 ArraySliceIndex
-		, uint32 NumArraySlices
-		, bool bUseIdentitySwizzle = false
-		, VkImageUsageFlags ImageUsageFlags = 0
-		, VkSamplerYcbcrConversion SamplerYcbcrConversion = nullptr);
+	FVulkanView* InitAsTextureView(VkImage InImage, const FVulkanTextureViewDesc& ViewDesc);
 
 	// NOTE: The InOffset applies to the FVulkanBuffer (it does not include any internal Allocation offsets that may exist)
 	FVulkanView* InitAsStructuredBufferView(
@@ -611,6 +617,8 @@ public:
 	// View with all mips/layers, but if it's a Depth/Stencil, only the Depth view
 	FVulkanView* PartialView = nullptr;
 
+	FVulkanView* FindOrAddInternalView(const FVulkanTextureViewDesc& ViewDesc);
+
 	FTextureRHIRef AliasedTexture;
 
 	template<typename T>
@@ -783,6 +791,9 @@ public:
 private:
 	void InternalMoveSurface(FVulkanDevice& InDevice, FVulkanCommandListContext& Context, VulkanRHI::FVulkanAllocation& DestAllocation);
 
+	static const VkImageUsageFlags DefaultImageUsageFlags = 0;
+	void InitViews(VkImageUsageFlags ViewUsageFlags, VkSamplerYcbcrConversion SamplerYcbcrConversion = VK_NULL_HANDLE);
+
 	VkImageTiling Tiling = VK_IMAGE_TILING_MAX_ENUM;
 	VulkanRHI::FVulkanAllocation Allocation;
 	VkImageAspectFlags FullAspectMask = 0;
@@ -790,6 +801,9 @@ private:
 	FVulkanCpuReadbackBuffer* CpuReadbackBuffer = nullptr;
 	VkImageLayout DefaultLayout = VK_IMAGE_LAYOUT_UNDEFINED;
 
+	// Views with lifetime tied to this image
+	TMap<uint32, FVulkanView*> InternalViews;
+
 protected:
 	EImageOwnerType ImageOwnerType;
 };
diff --git a/Engine/Source/Runtime/VulkanRHI/Private/VulkanTexture.cpp b/Engine/Source/Runtime/VulkanRHI/Private/VulkanTexture.cpp
--- a/Engine/Source/Runtime/VulkanRHI/Private/VulkanTexture.cpp
+++ b/Engine/Source/Runtime/VulkanRHI/Private/VulkanTexture.cpp
@@ -581,7 +581,6 @@ static VkImageLayout GetInitialLayoutFromRHIAccess(ERHIAccess RHIAccess, bool bI
 
 void FVulkanTexture::InternalMoveSurface(FVulkanDevice& InDevice, FVulkanCommandListContext& Context, VulkanRHI::FVulkanAllocation& DestAllocation)
 {
-
 	FImageCreateInfo ImageCreateInfo;
 	const FRHITextureDesc& Desc = GetDesc();
 	FVulkanTexture::GenerateImageCreateInfo(ImageCreateInfo, InDevice, Desc, &StorageFormat, &ViewFormat);
@@ -966,17 +965,18 @@ void FVulkanDynamicRHI::RHIUpdateTextureReference(FRHICommandListBase& RHICmdLis
 				const FRHITextureDesc& Desc = NewVulkanTexture->GetDesc();
 
 				VulkanTextureRefSRV->Invalidate();
-				VulkanTextureRefSRV->InitAsTextureView(
-					NewVulkanTexture->Image
-					, NewVulkanTexture->GetViewType()
-					, NewVulkanTexture->GetPartialAspectMask()
-					, Desc.Format
-					, NewVulkanTexture->ViewFormat
-					, 0u
-					, FMath::Max(Desc.NumMips, (uint8)1u)
-					, 0u
-					, NewVulkanTexture->GetNumberOfArrayLevels()
-					, !NewVulkanTexture->SupportsSampling());
+
+				FVulkanTextureViewDesc VulkanViewDesc;
+				VulkanViewDesc.ViewType = NewVulkanTexture->GetViewType();
+				VulkanViewDesc.AspectFlags = NewVulkanTexture->GetPartialAspectMask();
+				VulkanViewDesc.UEFormat = Desc.Format;
+				VulkanViewDesc.Format = NewVulkanTexture->ViewFormat;
+				VulkanViewDesc.FirstMip = 0u;
+				VulkanViewDesc.NumMips = FMath::Max(Desc.NumMips, (uint8)1u);
+				VulkanViewDesc.ArraySliceIndex = 0u;
+				VulkanViewDesc.NumArraySlices = NewVulkanTexture->GetNumberOfArrayLevels();
+				VulkanViewDesc.bUseIdentitySwizzle = !NewVulkanTexture->SupportsSampling();
+				VulkanTextureRefSRV->InitAsTextureView(NewVulkanTexture->Image, VulkanViewDesc);
 			}
 		});
 	}
@@ -1691,45 +1691,7 @@ FVulkanTexture::FVulkanTexture(FVulkanDevice& InDevice, const FRHITextureCreateD
 		checkf(StorageFormat != VK_FORMAT_UNDEFINED, TEXT("Pixel Format %d not defined!"), (int32)InCreateDesc.Format);
 	}
 
-	const VkDescriptorType DescriptorType = SupportsSampling() ? VK_DESCRIPTOR_TYPE_SAMPLED_IMAGE : VK_DESCRIPTOR_TYPE_STORAGE_IMAGE;
-	const bool bHasUAVFormat = (InCreateDesc.UAVFormat != PF_Unknown && InCreateDesc.UAVFormat != InCreateDesc.Format);
-	const VkImageUsageFlags SRVUsage = (bIsSRGB || bHasUAVFormat) ? (ImageCreateInfo.ImageCreateInfo.usage & ~VK_IMAGE_USAGE_STORAGE_BIT) : ImageCreateInfo.ImageCreateInfo.usage;
-	if (ViewType != VK_IMAGE_VIEW_TYPE_MAX_ENUM) //-V547
-	{
-		DefaultView = (new FVulkanView(InDevice, DescriptorType))->InitAsTextureView(
-			Image
-			, ViewType
-			, GetFullAspectMask()
-			, InCreateDesc.Format
-			, ViewFormat
-			, 0
-			, FMath::Max(InCreateDesc.NumMips, (uint8)1u)
-			, 0
-			, GetNumberOfArrayLevels()
-			, !SupportsSampling()
-			, SRVUsage
-		);
-	}
-
-	if (FullAspectMask == PartialAspectMask)
-	{
-		PartialView = DefaultView;
-	}
-	else
-	{
-		PartialView = (new FVulkanView(InDevice, DescriptorType))->InitAsTextureView(
-			Image
-			, ViewType
-			, PartialAspectMask
-			, InCreateDesc.Format
-			, ViewFormat
-			, 0
-			, FMath::Max(InCreateDesc.NumMips, (uint8)1u)
-			, 0
-			, GetNumberOfArrayLevels()
-			, false
-		);
-	}
+	InitViews(ImageUsageFlags);
 }
 
 FVulkanTexture::FVulkanTexture(FVulkanDevice& InDevice, const FRHITextureCreateDesc& InCreateDesc, VkImage InImage, const FVulkanRHIExternalImageDeleteCallbackInfo& InExternalImageDeleteCallbackInfo)
@@ -1771,52 +1733,13 @@ FVulkanTexture::FVulkanTexture(FVulkanDevice& InDevice, const FRHITextureCreateD
 
 			if (!EnumHasAnyFlags(InCreateDesc.Flags, TexCreate_Presentable))
 			{
-			FRHICommandList& RHICmdList = FRHICommandListExecutor::GetImmediateCommandList();
-			SetInitialImageState(RHICmdList, InitialLayout, bDoInitialClear, InCreateDesc.ClearValue, false);
+				FRHICommandList& RHICmdList = FRHICommandListExecutor::GetImmediateCommandList();
+				SetInitialImageState(RHICmdList, InitialLayout, bDoInitialClear, InCreateDesc.ClearValue, false);
+			}
 		}
 	}
-	}
-
-	const VkImageViewType ViewType = GetViewType();
-	const VkDescriptorType DescriptorType = SupportsSampling() ? VK_DESCRIPTOR_TYPE_SAMPLED_IMAGE : VK_DESCRIPTOR_TYPE_STORAGE_IMAGE;
-	const bool bUseIdentitySwizzle = (DescriptorType != VK_DESCRIPTOR_TYPE_SAMPLED_IMAGE) ||
-									(ViewFormat == VK_FORMAT_UNDEFINED); // External buffer textures also require identity swizzle
 
-	if (Image != VK_NULL_HANDLE)
-	{
-		DefaultView = (new FVulkanView(InDevice, DescriptorType))->InitAsTextureView(
-			Image
-			, ViewType
-			, GetFullAspectMask()
-			, InCreateDesc.Format
-			, ViewFormat
-			, 0
-			, FMath::Max(InCreateDesc.NumMips, (uint8)1u)
-			, 0
-			, GetNumberOfArrayLevels()
-			, bUseIdentitySwizzle
-		);
-	}
-
-	if (FullAspectMask == PartialAspectMask)
-	{
-		PartialView = DefaultView;
-	}
-	else
-	{
-		PartialView = (new FVulkanView(InDevice, DescriptorType))->InitAsTextureView(
-			Image
-			, ViewType
-			, PartialAspectMask
-			, InCreateDesc.Format
-			, ViewFormat
-			, 0
-			, FMath::Max(InCreateDesc.NumMips, (uint8)1u)
-			, 0
-			, GetNumberOfArrayLevels()
-			, false
-		);
-	}
+	InitViews(DefaultImageUsageFlags);
 }
 
 #if PLATFORM_ANDROID
@@ -2013,53 +1936,97 @@ FVulkanTexture::FVulkanTexture(FVulkanDevice& InDevice, const FRHITextureCreateD
 		}
 	}
 
+	InitViews(DefaultImageUsageFlags, SamplerYcbcrConversion);
+}
+
+#endif // PLATFORM_ANDROID
+
+
+void FVulkanTexture::InitViews(VkImageUsageFlags ViewUsageFlags, VkSamplerYcbcrConversion SamplerYcbcrConversion)
+{
+	const FRHITextureDesc& Desc = GetDesc();
 	const VkImageViewType ViewType = GetViewType();
 	const VkDescriptorType DescriptorType = SupportsSampling() ? VK_DESCRIPTOR_TYPE_SAMPLED_IMAGE : VK_DESCRIPTOR_TYPE_STORAGE_IMAGE;
-	const bool bUseIdentitySwizzle = (DescriptorType != VK_DESCRIPTOR_TYPE_SAMPLED_IMAGE) ||
-									(ViewFormat == VK_FORMAT_UNDEFINED); // External buffer textures also require identity swizzle
+	const bool bIsSRGB = EnumHasAllFlags(Desc.Flags, TexCreate_SRGB);
+	const bool bHasUAVFormat = (Desc.UAVFormat != PF_Unknown) && (Desc.UAVFormat != Desc.Format);
+	const bool bRemoveStorageBit = (bIsSRGB || (bHasUAVFormat && !UE::PixelFormat::HasCapabilities(Desc.Format, EPixelFormatCapabilities::UAV)));
+	if (bRemoveStorageBit)
+	{
+		ViewUsageFlags = (ViewUsageFlags & ~VK_IMAGE_USAGE_STORAGE_BIT);
+	}
+
+	const bool bUseIdentitySwizzle = 
+		(DescriptorType != VK_DESCRIPTOR_TYPE_SAMPLED_IMAGE) ||
+		(ViewFormat == VK_FORMAT_UNDEFINED); // External buffer textures also require identity swizzle
+
+	FVulkanTextureViewDesc ViewDesc;
+	ViewDesc.ViewType = GetViewType();
+	ViewDesc.AspectFlags = FullAspectMask;
+	ViewDesc.UEFormat = Desc.Format;
+	ViewDesc.Format = ViewFormat;
+	ViewDesc.FirstMip = 0;
+	ViewDesc.NumMips = FMath::Max(Desc.NumMips, (uint8)1u);
+	ViewDesc.ArraySliceIndex = 0;
+	ViewDesc.NumArraySlices = GetNumberOfArrayLevels();
+	ViewDesc.bUseIdentitySwizzle = bUseIdentitySwizzle;
+	ViewDesc.ImageUsageFlags = ViewUsageFlags;
+	ViewDesc.SamplerYcbcrConversion = SamplerYcbcrConversion;
+
+	if ((Image != VK_NULL_HANDLE) && (ViewType != VK_IMAGE_VIEW_TYPE_MAX_ENUM))
+	{
+		if (DefaultView)
+		{
+			DefaultView->Invalidate();
+			DefaultView->InitAsTextureView(Image, ViewDesc);
+		}
+		else
+		{
+			DefaultView = (new FVulkanView(*Device, DescriptorType))->InitAsTextureView(Image, ViewDesc);
+		}
 
-	if (Image != VK_NULL_HANDLE)
-	{
-		DefaultView = (new FVulkanView(InDevice, DescriptorType))->InitAsTextureView(
-			Image
-			, ViewType
-			, GetFullAspectMask()
-			, InCreateDesc.Format
-			, ViewFormat
-			, 0
-			, FMath::Max(InCreateDesc.NumMips, (uint8)1u)
-			, 0
-			, GetNumberOfArrayLevels()
-			, bUseIdentitySwizzle
-			, 0
-			, SamplerYcbcrConversion
-		);
+		if (FullAspectMask == PartialAspectMask)
+		{
+			PartialView = DefaultView;
+		}
+		else
+		{
+			ViewDesc.AspectFlags = PartialAspectMask;
+			ViewDesc.bUseIdentitySwizzle = false;
+			ViewDesc.SamplerYcbcrConversion = VK_NULL_HANDLE;
+
+			if (PartialView)
+			{
+				PartialView->Invalidate();
+				PartialView->InitAsTextureView(Image, ViewDesc);
+			}
+			else
+			{
+				PartialView = (new FVulkanView(*Device, DescriptorType))->InitAsTextureView(Image, ViewDesc);
+			}
+		}
 	}
+}
 
-	if (FullAspectMask == PartialAspectMask)
+static FRWLock GInternalViewRWLock; // Use single global lock for now
+FVulkanView* FVulkanTexture::FindOrAddInternalView(const FVulkanTextureViewDesc& ViewDesc)
+{
 	{
-		PartialView = DefaultView;
+		FRWScopeLock ScopedLock(GInternalViewRWLock, SLT_ReadOnly);
+		if (FVulkanView** VulkanView = InternalViews.Find(ViewDesc.GetHash()))
+		{
+			return *VulkanView;
+		}
 	}
-	else
+
+	FRWScopeLock ScopedLock(GInternalViewRWLock, SLT_Write);
+	FVulkanView*& VulkanView = InternalViews.FindOrAdd(ViewDesc.GetHash(), nullptr);
+	if (!VulkanView)
 	{
-		PartialView = (new FVulkanView(InDevice, DescriptorType))->InitAsTextureView(
-			Image
-			, ViewType
-			, PartialAspectMask
-			, InCreateDesc.Format
-			, ViewFormat
-			, 0
-			, FMath::Max(InCreateDesc.NumMips, (uint8)1u)
-			, 0
-			, GetNumberOfArrayLevels()
-			, false
-		);
+		VulkanView = (new FVulkanView(*Device, VK_DESCRIPTOR_TYPE_SAMPLED_IMAGE))->InitAsTextureView(Image, ViewDesc);
 	}
+	return VulkanView;
 }
 
-#endif // PLATFORM_ANDROID
-
-
 FVulkanTexture::FVulkanTexture(FVulkanDevice& InDevice, const FRHITextureCreateDesc& InCreateDesc, FTextureRHIRef& SrcTextureRHI)
 	: FRHITexture(InCreateDesc)
 	, Device(&InDevice)
@@ -2102,6 +2069,13 @@ FVulkanTexture::~FVulkanTexture()
 		}
 
 		delete DefaultView;
+
+		for (auto& Pair : InternalViews)
+		{
+			delete Pair.Value;
+		}
+		InternalViews.Reset();
+		
 		DestroySurface();
 	}
 }
@@ -2119,22 +2093,19 @@ void FVulkanTexture::AliasTextureResources(FTextureRHIRef& SrcTextureRHI)
 
 void FVulkanTexture::UpdateLinkedViews()
 {
-	DefaultView->Invalidate();
+	InitViews(ImageUsageFlags);
 
-	const FRHITextureDesc& Desc = GetDesc();
-	const uint32 NumMips = Desc.NumMips;
-	const VkImageViewType ViewType = GetViewType();
-	const uint32 ArraySize = GetNumberOfArrayLevels();
-
-	if (ViewType != VK_IMAGE_VIEW_TYPE_MAX_ENUM) //-V547
+	const bool bRenderTarget = EnumHasAnyFlags(GetDesc().Flags, TexCreate_RenderTargetable | TexCreate_DepthStencilTargetable | TexCreate_ResolveTargetable);
+	if (bRenderTarget)
 	{
-		DefaultView->InitAsTextureView(Image, ViewType, GetFullAspectMask(), GetDesc().Format, ViewFormat, 0, FMath::Max(NumMips, 1u), 0, ArraySize, !SupportsSampling());
+		Device->NotifyDeletedImage(Image, true);
 	}
-	if (PartialView != DefaultView)
+
+	for (auto& Pair : InternalViews)
 	{
-		PartialView->Invalidate();
-		PartialView->InitAsTextureView(Image, ViewType, PartialAspectMask, GetDesc().Format, ViewFormat, 0, FMath::Max(NumMips, 1u), 0, ArraySize, false);
+		delete Pair.Value;
 	}
+	InternalViews.Reset();
 
 	FVulkanViewableResource::UpdateLinkedViews();
 }
diff --git a/Engine/Source/Runtime/VulkanRHI/Private/VulkanUAV.cpp b/Engine/Source/Runtime/VulkanRHI/Private/VulkanUAV.cpp
--- a/Engine/Source/Runtime/VulkanRHI/Private/VulkanUAV.cpp
+++ b/Engine/Source/Runtime/VulkanRHI/Private/VulkanUAV.cpp
@@ -7,6 +7,26 @@
 #include "ClearReplacementShaders.h"
 #include "VulkanRayTracing.h"
 
+
+uint32 FVulkanTextureViewDesc::GetHash() const
+{
+	uint32 Hash = 0;
+	Hash = FCrc::TypeCrc32(ViewType, Hash);
+	Hash = FCrc::TypeCrc32(AspectFlags, Hash);
+	Hash = FCrc::TypeCrc32(UEFormat, Hash);
+	Hash = FCrc::TypeCrc32(Format, Hash);
+	Hash = FCrc::TypeCrc32(FirstMip, Hash);
+	Hash = FCrc::TypeCrc32(NumMips, Hash);
+	Hash = FCrc::TypeCrc32(ArraySliceIndex, Hash);
+	Hash = FCrc::TypeCrc32(NumArraySlices, Hash);
+	Hash = FCrc::TypeCrc32(bUseIdentitySwizzle, Hash);
+	Hash = FCrc::TypeCrc32(ImageUsageFlags, Hash);
+	Hash = FCrc::TypeCrc32(SamplerYcbcrConversion, Hash);
+	return Hash;
+}
+
+
+
 FVulkanView::FVulkanView(FVulkanDevice& InDevice, VkDescriptorType InDescriptorType)
 	: Device(InDevice)
 	, DescriptorType(InDescriptorType)
@@ -112,19 +132,7 @@ FVulkanView* FVulkanView::InitAsTypedBufferView(FVulkanBuffer* Buffer, EPixelFor
 	return this;
 }
 
-FVulkanView* FVulkanView::InitAsTextureView(
-	  VkImage InImage
-	, VkImageViewType ViewType
-	, VkImageAspectFlags AspectFlags
-	, EPixelFormat UEFormat
-	, VkFormat Format
-	, uint32 FirstMip
-	, uint32 NumMips
-	, uint32 ArraySliceIndex
-	, uint32 NumArraySlices
-	, bool bUseIdentitySwizzle
-	, VkImageUsageFlags ImageUsageFlags
-	, VkSamplerYcbcrConversion SamplerYcbcrConversion)
+FVulkanView* FVulkanView::InitAsTextureView(VkImage InImage, const FVulkanTextureViewDesc& ViewDesc)
 {
 	// We will need a deferred update if the descriptor was already in use
 	const bool bImmediateUpdate = !IsInitialized();
@@ -138,12 +146,12 @@ FVulkanView* FVulkanView::InitAsTextureView(
 	VkImageViewCreateInfo ViewInfo;
 	ZeroVulkanStruct(ViewInfo, VK_STRUCTURE_TYPE_IMAGE_VIEW_CREATE_INFO);
 	ViewInfo.image = InImage;
-	ViewInfo.viewType = ViewType;
-	ViewInfo.format = Format;
+	ViewInfo.viewType = ViewDesc.ViewType;
+	ViewInfo.format = ViewDesc.Format;
 
 #if VULKAN_SUPPORTS_ASTC_DECODE_MODE
 	VkImageViewASTCDecodeModeEXT DecodeMode;
-	if (Device.GetOptionalExtensions().HasEXTASTCDecodeMode && IsAstcLdrFormat(Format) && !IsAstcSrgbFormat(Format))
+	if (Device.GetOptionalExtensions().HasEXTASTCDecodeMode && IsAstcLdrFormat(ViewDesc.Format) && !IsAstcSrgbFormat(ViewDesc.Format))
 	{
 		ZeroVulkanStruct(DecodeMode, VK_STRUCTURE_TYPE_IMAGE_VIEW_ASTC_DECODE_MODE_EXT);
 		DecodeMode.decodeMode = VK_FORMAT_R8G8B8A8_UNORM;
@@ -153,15 +161,15 @@ FVulkanView* FVulkanView::InitAsTextureView(
 #endif
 
 	VkSamplerYcbcrConversionInfo SamplerYcbcrConversionInfo;
-	if (SamplerYcbcrConversion)
+	if (ViewDesc.SamplerYcbcrConversion)
 	{
 		ZeroVulkanStruct(SamplerYcbcrConversionInfo, VK_STRUCTURE_TYPE_SAMPLER_YCBCR_CONVERSION_INFO);
-		SamplerYcbcrConversionInfo.conversion = SamplerYcbcrConversion;
+		SamplerYcbcrConversionInfo.conversion = ViewDesc.SamplerYcbcrConversion;
 		SamplerYcbcrConversionInfo.pNext = ViewInfo.pNext;
 		ViewInfo.pNext = &SamplerYcbcrConversionInfo;
 	}
 
-	if (bUseIdentitySwizzle)
+	if (ViewDesc.bUseIdentitySwizzle)
 	{
 		ViewInfo.components.a = VK_COMPONENT_SWIZZLE_IDENTITY;
 		ViewInfo.components.r = VK_COMPONENT_SWIZZLE_IDENTITY;
@@ -170,23 +178,23 @@ FVulkanView* FVulkanView::InitAsTextureView(
 	}
 	else
 	{
-		ViewInfo.components = Device.GetFormatComponentMapping(UEFormat);
+		ViewInfo.components = Device.GetFormatComponentMapping(ViewDesc.UEFormat);
 	}
 
-	ViewInfo.subresourceRange.aspectMask = AspectFlags;
-	ViewInfo.subresourceRange.baseMipLevel = FirstMip;
-	ensure(NumMips != 0xFFFFFFFF);
-	ViewInfo.subresourceRange.levelCount = NumMips;
+	ViewInfo.subresourceRange.aspectMask = ViewDesc.AspectFlags;
+	ViewInfo.subresourceRange.baseMipLevel = ViewDesc.FirstMip;
+	ensure(ViewDesc.NumMips != 0xFFFFFFFF);
+	ViewInfo.subresourceRange.levelCount = ViewDesc.NumMips;
 
-	ensure(ArraySliceIndex != 0xFFFFFFFF);
-	ensure(NumArraySlices != 0xFFFFFFFF);
-	ViewInfo.subresourceRange.baseArrayLayer = ArraySliceIndex;
-	ViewInfo.subresourceRange.layerCount = NumArraySlices;
+	ensure(ViewDesc.ArraySliceIndex != 0xFFFFFFFF);
+	ensure(ViewDesc.NumArraySlices != 0xFFFFFFFF);
+	ViewInfo.subresourceRange.baseArrayLayer = ViewDesc.ArraySliceIndex;
+	ViewInfo.subresourceRange.layerCount = ViewDesc.NumArraySlices;
 
 	//HACK.  DX11 on PC currently uses a D24S8 depthbuffer and so needs an X24_G8 SRV to visualize stencil.
 	//So take that as our cue to visualize stencil.  In the future, the platform independent code will have a real format
 	//instead of PF_DepthStencil, so the cross-platform code could figure out the proper format to pass in for this.
-	if (UEFormat == PF_X24_G8)
+	if (ViewDesc.UEFormat == PF_X24_G8)
 	{
 		ensure((ViewInfo.format == (VkFormat)GPixelFormats[PF_DepthStencil].PlatformFormat) && (ViewInfo.format != VK_FORMAT_UNDEFINED));
 		ViewInfo.subresourceRange.aspectMask = VK_IMAGE_ASPECT_STENCIL_BIT;
@@ -194,10 +202,10 @@ FVulkanView* FVulkanView::InitAsTextureView(
 
 	// Inform the driver the view will only be used with a subset of usage flags (to help performance and/or compatibility)
 	VkImageViewUsageCreateInfo ImageViewUsageCreateInfo;
-	if (ImageUsageFlags != 0)
+	if (ViewDesc.ImageUsageFlags != 0)
 	{
 		ZeroVulkanStruct(ImageViewUsageCreateInfo, VK_STRUCTURE_TYPE_IMAGE_VIEW_USAGE_CREATE_INFO);
-		ImageViewUsageCreateInfo.usage = ImageUsageFlags;
+		ImageViewUsageCreateInfo.usage = ViewDesc.ImageUsageFlags;
 
 		ImageViewUsageCreateInfo.pNext = (void*)ViewInfo.pNext;
 		ViewInfo.pNext = &ImageViewUsageCreateInfo;
@@ -213,7 +221,7 @@ FVulkanView* FVulkanView::InitAsTextureView(
 		TV.ViewId = ++GVulkanImageViewHandleIdCounter;
 	}
 
-	const bool bDepthOrStencilAspect = (AspectFlags & (VK_IMAGE_ASPECT_DEPTH_BIT | VK_IMAGE_ASPECT_STENCIL_BIT)) != 0;
+	const bool bDepthOrStencilAspect = (ViewDesc.AspectFlags & (VK_IMAGE_ASPECT_DEPTH_BIT | VK_IMAGE_ASPECT_STENCIL_BIT)) != 0;
 	Device.GetBindlessDescriptorManager()->UpdateImage(BindlessHandle, DescriptorType, TV.View, bDepthOrStencilAspect, bImmediateUpdate);
 
 	return this;
@@ -450,19 +458,18 @@ void FVulkanShaderResourceView::UpdateView()
 		// Remove storage bit from usage flags for SRVs
 		const VkImageUsageFlags ImageUsageFlags = Texture->ImageUsageFlags & ~VK_IMAGE_USAGE_STORAGE_BIT;
 
-		InitAsTextureView(
-			  Texture->Image
-			, GetVkImageViewTypeForDimensionSRV(Info.Dimension, Texture->GetViewType())
-			, Texture->GetPartialAspectMask()
-			, Info.Format
-			, UEToVkTextureFormat(Info.Format, Info.bSRGB)
-			, Info.MipRange.First
-			, Info.MipRange.Num
-			, ArrayFirst
-			, ArrayNum
-			, false
-			, ImageUsageFlags
-		);
+		FVulkanTextureViewDesc VulkanViewDesc;
+		VulkanViewDesc.ViewType = GetVkImageViewTypeForDimensionSRV(Info.Dimension, Texture->GetViewType());
+		VulkanViewDesc.AspectFlags = Texture->GetPartialAspectMask();
+		VulkanViewDesc.UEFormat = Info.Format;
+		VulkanViewDesc.Format = UEToVkTextureFormat(Info.Format, Info.bSRGB);
+		VulkanViewDesc.FirstMip = Info.MipRange.First;
+		VulkanViewDesc.NumMips = Info.MipRange.Num;
+		VulkanViewDesc.ArraySliceIndex = ArrayFirst;
+		VulkanViewDesc.NumArraySlices = ArrayNum;
+		VulkanViewDesc.bUseIdentitySwizzle = false;
+		VulkanViewDesc.ImageUsageFlags = ImageUsageFlags;
+		InitAsTextureView(Texture->Image, VulkanViewDesc);
 	}
 }
 
@@ -541,18 +548,17 @@ void FVulkanUnorderedAccessView::UpdateView()
 			checkf((ArrayFirst + ArrayNum) <= Texture->GetNumberOfArrayLevels(), TEXT("View extends beyond original cube texture level count!"));
 		}
 
-		InitAsTextureView(
-			  Texture->Image
-			, GetVkImageViewTypeForDimensionUAV(Info.Dimension, Texture->GetViewType())
-			, Texture->GetPartialAspectMask()
-			, Info.Format
-			, UEToVkTextureFormat(Info.Format, false)
-			, Info.MipLevel
-			, 1
-			, ArrayFirst
-			, ArrayNum
-			, true
-		);
+		FVulkanTextureViewDesc VulkanViewDesc;
+		VulkanViewDesc.ViewType = GetVkImageViewTypeForDimensionUAV(Info.Dimension, Texture->GetViewType());
+		VulkanViewDesc.AspectFlags = Texture->GetPartialAspectMask();
+		VulkanViewDesc.UEFormat = Info.Format;
+		VulkanViewDesc.Format = UEToVkTextureFormat(Info.Format, false);
+		VulkanViewDesc.FirstMip = Info.MipLevel;
+		VulkanViewDesc.NumMips = 1;
+		VulkanViewDesc.ArraySliceIndex = ArrayFirst;
+		VulkanViewDesc.NumArraySlices = ArrayNum;
+		VulkanViewDesc.bUseIdentitySwizzle = true;
+		InitAsTextureView(Texture->Image, VulkanViewDesc);
 	}
 }
 
diff --git a/Engine/Source/Runtime/VulkanRHI/Private/VulkanViewport.cpp b/Engine/Source/Runtime/VulkanRHI/Private/VulkanViewport.cpp
--- a/Engine/Source/Runtime/VulkanRHI/Private/VulkanViewport.cpp
+++ b/Engine/Source/Runtime/VulkanRHI/Private/VulkanViewport.cpp
@@ -302,20 +302,6 @@ FVulkanFramebuffer::FVulkanFramebuffer(FVulkanDevice& Device, const FRHISetRende
 
 	AttachmentTextureViews.Empty(RTLayout.GetNumAttachmentDescriptions());
 
-	auto CreateOwnedView = [&]()
-	{
-		const VkDescriptorType DescriptorType = VK_DESCRIPTOR_TYPE_SAMPLED_IMAGE;
-		FVulkanView* View = new FVulkanView(Device, DescriptorType);
-		AttachmentTextureViews.Add(View);
-		OwnedTextureViews.Add(View);
-		return View;
-	};
-
-	auto AddExternalView = [&](FVulkanView const* View)
-	{
-		AttachmentTextureViews.Add(View);
-	};
-
 	uint32 MipIndex = 0;
 
 	const VkExtent3D& RTExtents = RTLayout.GetExtent3D();
@@ -332,92 +318,12 @@ FVulkanFramebuffer::FVulkanFramebuffer(FVulkanDevice& Device, const FRHISetRende
 		}
 
 		FVulkanTexture* Texture = ResourceCast(RHITexture);
-		const FRHITextureDesc& Desc = Texture->GetDesc();
-
-		// this could fire in case one of the textures is FVulkanBackBuffer and it has not acquired an image
-		// with EDelayAcquireImageType::LazyAcquire acquire happens when texture transition to Writeable state
-		// make sure you call TransitionResource(Writable, Tex) before using this texture as a render-target
-		check(Texture->Image != VK_NULL_HANDLE);
 
 		ColorRenderTargetImages[Index] = Texture->Image;
 		MipIndex = InRTInfo.ColorRenderTarget[Index].MipIndex;
 
-		if (Texture->GetViewType() == VK_IMAGE_VIEW_TYPE_2D || Texture->GetViewType() == VK_IMAGE_VIEW_TYPE_2D_ARRAY)
-		{
-			uint32 ArraySliceIndex = 0;
-			uint32 NumArraySlices = 1;
-			if (InRTInfo.ColorRenderTarget[Index].ArraySliceIndex == -1)
-			{
-				ArraySliceIndex = 0;
-				NumArraySlices = Texture->GetNumberOfArrayLevels();
-			}
-			else
-			{
-				ArraySliceIndex = InRTInfo.ColorRenderTarget[Index].ArraySliceIndex;
-				NumArraySlices = 1;
-				check(ArraySliceIndex < Texture->GetNumberOfArrayLevels());
-			}
-
-			// About !RTLayout.GetIsMultiView(), from https://registry.khronos.org/vulkan/specs/1.3-extensions/man/html/VkFramebufferCreateInfo.html: 
-			// If the render pass uses multiview, then layers must be one
-			if (Texture->GetViewType() == VK_IMAGE_VIEW_TYPE_2D_ARRAY && !RTLayout.GetIsMultiView())
-			{
-				NumLayers = NumArraySlices;
-			}
-
-			CreateOwnedView()->InitAsTextureView(
-				  Texture->Image
-				, Texture->GetViewType()
-				, Texture->GetFullAspectMask()
-				, Desc.Format
-				, Texture->ViewFormat
-				, MipIndex
-				, 1
-				, ArraySliceIndex
-				, NumArraySlices
-				, true
-				, VK_IMAGE_USAGE_COLOR_ATTACHMENT_BIT | (Texture->ImageUsageFlags & VK_IMAGE_USAGE_INPUT_ATTACHMENT_BIT)
-			);
-		}
-		else if (Texture->GetViewType() == VK_IMAGE_VIEW_TYPE_CUBE || Texture->GetViewType() == VK_IMAGE_VIEW_TYPE_CUBE_ARRAY)
-		{
-			// Cube always renders one face at a time
-			INC_DWORD_STAT(STAT_VulkanNumImageViews);
-
-			CreateOwnedView()->InitAsTextureView(
-				  Texture->Image
-				, VK_IMAGE_VIEW_TYPE_2D
-				, Texture->GetFullAspectMask()
-				, Desc.Format
-				, Texture->ViewFormat
-				, MipIndex
-				, 1
-				, InRTInfo.ColorRenderTarget[Index].ArraySliceIndex
-				, 1
-				, true
-				, VK_IMAGE_USAGE_COLOR_ATTACHMENT_BIT | (Texture->ImageUsageFlags & VK_IMAGE_USAGE_INPUT_ATTACHMENT_BIT)
-			);
-		}
-		else if (Texture->GetViewType() == VK_IMAGE_VIEW_TYPE_3D)
-		{
-			CreateOwnedView()->InitAsTextureView(
-				  Texture->Image
-				, VK_IMAGE_VIEW_TYPE_2D_ARRAY
-				, Texture->GetFullAspectMask()
-				, Desc.Format
-				, Texture->ViewFormat
-				, MipIndex
-				, 1
-				, 0
-				, Desc.Depth
-				, true
-				, VK_IMAGE_USAGE_COLOR_ATTACHMENT_BIT | (Texture->ImageUsageFlags & VK_IMAGE_USAGE_INPUT_ATTACHMENT_BIT)
-			);
-		}
-		else
-		{
-			ensure(0);
-		}
+		FVulkanView* View = GetColorRenderTargetViewDesc(Texture, MipIndex, InRTInfo.ColorRenderTarget[Index].ArraySliceIndex, RTLayout.GetMultiViewCount(), NumLayers);
+		AttachmentTextureViews.Add(View);
 
 		++NumColorAttachments;
 
@@ -429,21 +335,9 @@ FVulkanFramebuffer::FVulkanFramebuffer(FVulkanDevice& Device, const FRHISetRende
 			FVulkanTexture* ResolveTexture = ResourceCast(ResolveRHITexture);
 			ColorResolveTargetImages[Index] = ResolveTexture->Image;
 
-			//resolve attachments only supported for 2d/2d array textures
-			if (ResolveTexture->GetViewType() == VK_IMAGE_VIEW_TYPE_2D || ResolveTexture->GetViewType() == VK_IMAGE_VIEW_TYPE_2D_ARRAY)
+			if (FVulkanView* ResolveView = GetColorResolveTargetViewDesc(ResolveTexture, MipIndex, InRTInfo.ColorRenderTarget[Index].ArraySliceIndex))
 			{
-				CreateOwnedView()->InitAsTextureView(
-					  ResolveTexture->Image
-					, ResolveTexture->GetViewType()
-					, ResolveTexture->GetFullAspectMask()
-					, ResolveTexture->GetDesc().Format
-					, ResolveTexture->ViewFormat
-					, MipIndex
-					, 1
-					, FMath::Max(0, (int32)InRTInfo.ColorRenderTarget[Index].ArraySliceIndex)
-					, ResolveTexture->GetNumberOfArrayLevels()
-					, true
-				);
+				AttachmentTextureViews.Add(ResolveView);
 			}
 		}
 	}
@@ -453,49 +347,12 @@ FVulkanFramebuffer::FVulkanFramebuffer(FVulkanDevice& Device, const FRHISetRende
 		FVulkanTexture* Texture = ResourceCast(InRTInfo.DepthStencilRenderTarget.Texture);
 		const FRHITextureDesc& Desc = Texture->GetDesc();
 		DepthStencilRenderTargetImage = Texture->Image;
-		bool bHasStencil = (Texture->GetDesc().Format == PF_DepthStencil || Texture->GetDesc().Format == PF_X24_G8);
 
 		check(Texture->PartialView);
 		PartialDepthTextureView = Texture->PartialView;
 
-		ensure(Texture->GetViewType() == VK_IMAGE_VIEW_TYPE_2D || Texture->GetViewType() == VK_IMAGE_VIEW_TYPE_2D_ARRAY || Texture->GetViewType() == VK_IMAGE_VIEW_TYPE_CUBE);
-		if (NumColorAttachments == 0 && Texture->GetViewType() == VK_IMAGE_VIEW_TYPE_CUBE)
-		{
-			CreateOwnedView()->InitAsTextureView(
-				  Texture->Image
-				, VK_IMAGE_VIEW_TYPE_2D_ARRAY
-				, Texture->GetFullAspectMask()
-				, Texture->GetDesc().Format
-				, Texture->ViewFormat
-				, MipIndex
-				, 1
-				, 0
-				, 6
-				, true
-			);
-
-			NumLayers = 6;
-		}
-		else if (Texture->GetViewType() == VK_IMAGE_VIEW_TYPE_2D  || Texture->GetViewType() == VK_IMAGE_VIEW_TYPE_2D_ARRAY)
-		{
-			// depth attachments need a separate view to have no swizzle components, for validation correctness
-			CreateOwnedView()->InitAsTextureView(
-				  Texture->Image
-				, Texture->GetViewType()
-				, Texture->GetFullAspectMask()
-				, Texture->GetDesc().Format
-				, Texture->ViewFormat
-				, MipIndex
-				, 1
-				, 0
-				, Texture->GetNumberOfArrayLevels()
-				, true
-			);
-		}
-		else
-		{
-			AddExternalView(Texture->DefaultView);
-		}
+		FVulkanView* View = GetDepthStencilTargetViewDesc(Texture, InRTInfo.NumColorRenderTargets, MipIndex, NumLayers);
+		AttachmentTextureViews.Add(View);
 
 		if (RTLayout.GetHasDepthStencilResolve() && RTLayout.GetDepthStencilResolveAttachmentReference()->layout != VK_IMAGE_LAYOUT_UNDEFINED)
 		{
@@ -503,21 +360,9 @@ FVulkanFramebuffer::FVulkanFramebuffer(FVulkanDevice& Device, const FRHISetRende
 			FVulkanTexture* ResolveTexture = ResourceCast(ResolveRHITexture);
 			DepthStencilResolveRenderTargetImage = ResolveTexture->Image;
 
-			// Resolve attachments only supported for 2d/2d array textures
-			if (ResolveTexture->GetViewType() == VK_IMAGE_VIEW_TYPE_2D || ResolveTexture->GetViewType() == VK_IMAGE_VIEW_TYPE_2D_ARRAY)
+			if (FVulkanView* ResolveView = GetDepthStencilResolveTargetViewDesc(ResolveTexture, MipIndex))
 			{
-				CreateOwnedView()->InitAsTextureView(
-					ResolveTexture->Image
-					, ResolveTexture->GetViewType()
-					, ResolveTexture->GetFullAspectMask()
-					, ResolveTexture->GetDesc().Format
-					, ResolveTexture->ViewFormat
-					, MipIndex
-					, 1
-					, 0
-					, ResolveTexture->GetNumberOfArrayLevels()
-					, true
-				);
+				AttachmentTextureViews.Add(ResolveView);
 			}
 		}
 	}
@@ -526,21 +371,8 @@ FVulkanFramebuffer::FVulkanFramebuffer(FVulkanDevice& Device, const FRHISetRende
 	{
 		FVulkanTexture* Texture = ResourceCast(InRTInfo.ShadingRateTexture);
 		FragmentDensityImage = Texture->Image;
-
-		ensure(Texture->GetViewType() == VK_IMAGE_VIEW_TYPE_2D || Texture->GetViewType() == VK_IMAGE_VIEW_TYPE_2D_ARRAY);
-
-		CreateOwnedView()->InitAsTextureView(
-			Texture->Image
-			, Texture->GetViewType()
-			, Texture->GetFullAspectMask()
-			, Texture->GetDesc().Format
-			, Texture->ViewFormat
-			, MipIndex
-			, 1
-			, 0
-			, Texture->GetNumberOfArrayLevels()
-			, true
-		);
+		FVulkanView* View = GetFragmentDensityAttachmentViewDesc(Texture, MipIndex);
+		AttachmentTextureViews.Add(View);
 	}
 
 	TArray<VkImageView> AttachmentViews;
@@ -569,6 +401,188 @@ FVulkanFramebuffer::FVulkanFramebuffer(FVulkanDevice& Device, const FRHISetRende
 	INC_DWORD_STAT(STAT_VulkanNumFrameBuffers);
 }
 
+FVulkanView* FVulkanFramebuffer::GetColorRenderTargetViewDesc(FVulkanTexture* Texture, uint32 MipIndex, int32 SliceIndex, uint32 MultiViewCount, uint32& InOutNumLayers)
+{
+	check(Texture->Image != VK_NULL_HANDLE);
+
+	const FRHITextureDesc& Desc = Texture->GetDesc();
+	const VkImageUsageFlags ImageUsageFlags = VK_IMAGE_USAGE_COLOR_ATTACHMENT_BIT | (Texture->ImageUsageFlags & VK_IMAGE_USAGE_INPUT_ATTACHMENT_BIT);
+
+	if (Texture->GetViewType() == VK_IMAGE_VIEW_TYPE_2D || Texture->GetViewType() == VK_IMAGE_VIEW_TYPE_2D_ARRAY)
+	{
+		uint32 ArraySliceIndex = 0;
+		uint32 NumArraySlices = 1;
+		if (SliceIndex == -1)
+		{
+			ArraySliceIndex = 0;
+			NumArraySlices = Texture->GetNumberOfArrayLevels();
+		}
+		else
+		{
+			ArraySliceIndex = SliceIndex;
+			NumArraySlices = 1;
+			check(ArraySliceIndex < Texture->GetNumberOfArrayLevels());
+		}
+
+		// About !RTLayout.GetIsMultiView(), from https://registry.khronos.org/vulkan/specs/1.3-extensions/man/html/VkFramebufferCreateInfo.html: 
+		// If the render pass uses multiview, then layers must be one
+		if (Texture->GetViewType() == VK_IMAGE_VIEW_TYPE_2D_ARRAY && (MultiViewCount != 0))
+		{
+			InOutNumLayers = NumArraySlices;
+		}
+
+		FVulkanTextureViewDesc ViewDesc;
+		ViewDesc.ViewType = Texture->GetViewType();
+		ViewDesc.AspectFlags = Texture->GetFullAspectMask();
+		ViewDesc.UEFormat = Desc.Format;
+		ViewDesc.Format = Texture->ViewFormat;
+		ViewDesc.FirstMip = MipIndex;
+		ViewDesc.NumMips = 1;
+		ViewDesc.ArraySliceIndex = ArraySliceIndex;
+		ViewDesc.NumArraySlices = NumArraySlices;
+		ViewDesc.bUseIdentitySwizzle = true;
+		ViewDesc.ImageUsageFlags = ImageUsageFlags;
+		return Texture->FindOrAddInternalView(ViewDesc);
+	}
+	else if (Texture->GetViewType() == VK_IMAGE_VIEW_TYPE_CUBE || Texture->GetViewType() == VK_IMAGE_VIEW_TYPE_CUBE_ARRAY)
+	{
+		// Cube always renders one face at a time
+		INC_DWORD_STAT(STAT_VulkanNumImageViews);
+
+		FVulkanTextureViewDesc ViewDesc;
+		ViewDesc.ViewType = VK_IMAGE_VIEW_TYPE_2D;
+		ViewDesc.AspectFlags = Texture->GetFullAspectMask();
+		ViewDesc.UEFormat = Desc.Format;
+		ViewDesc.Format = Texture->ViewFormat;
+		ViewDesc.FirstMip = MipIndex;
+		ViewDesc.NumMips = 1;
+		ViewDesc.ArraySliceIndex = SliceIndex;
+		ViewDesc.NumArraySlices = 1;
+		ViewDesc.bUseIdentitySwizzle = true;
+		ViewDesc.ImageUsageFlags = ImageUsageFlags;
+		return Texture->FindOrAddInternalView(ViewDesc);
+	}
+	else if (Texture->GetViewType() == VK_IMAGE_VIEW_TYPE_3D)
+	{
+		FVulkanTextureViewDesc ViewDesc;
+		ViewDesc.ViewType = VK_IMAGE_VIEW_TYPE_2D_ARRAY;
+		ViewDesc.AspectFlags = Texture->GetFullAspectMask();
+		ViewDesc.UEFormat = Desc.Format;
+		ViewDesc.Format = Texture->ViewFormat;
+		ViewDesc.FirstMip = MipIndex;
+		ViewDesc.NumMips = 1;
+		ViewDesc.ArraySliceIndex = 0;
+		ViewDesc.NumArraySlices = Desc.Depth;
+		ViewDesc.bUseIdentitySwizzle = true;
+		ViewDesc.ImageUsageFlags = ImageUsageFlags;
+		return Texture->FindOrAddInternalView(ViewDesc);
+	}
+	ensure(0);
+	return nullptr;
+}
+
+FVulkanView* FVulkanFramebuffer::GetColorResolveTargetViewDesc(FVulkanTexture* ResolveTexture, uint32 MipIndex, int32 ArraySliceIndex)
+{
+	//resolve attachments only supported for 2d/2d array textures
+	if (ResolveTexture->GetViewType() == VK_IMAGE_VIEW_TYPE_2D || ResolveTexture->GetViewType() == VK_IMAGE_VIEW_TYPE_2D_ARRAY)
+	{
+		FVulkanTextureViewDesc ViewDesc;
+		ViewDesc.ViewType = ResolveTexture->GetViewType();
+		ViewDesc.AspectFlags = ResolveTexture->GetFullAspectMask();
+		ViewDesc.UEFormat = ResolveTexture->GetDesc().Format;
+		ViewDesc.Format = ResolveTexture->ViewFormat;
+		ViewDesc.FirstMip = MipIndex;
+		ViewDesc.NumMips = 1;
+		ViewDesc.ArraySliceIndex = FMath::Max(0, ArraySliceIndex);
+		ViewDesc.NumArraySlices = ResolveTexture->GetNumberOfArrayLevels();
+		ViewDesc.bUseIdentitySwizzle = true;
+		return ResolveTexture->FindOrAddInternalView(ViewDesc);
+	}
+
+	return nullptr;
+}
+
+FVulkanView* FVulkanFramebuffer::GetDepthStencilTargetViewDesc(FVulkanTexture* Texture, uint32 NumColorAttachments, uint32 MipIndex, uint32& InOutNumLayers)
+{
+	ensure(Texture->GetViewType() == VK_IMAGE_VIEW_TYPE_2D || Texture->GetViewType() == VK_IMAGE_VIEW_TYPE_2D_ARRAY || Texture->GetViewType() == VK_IMAGE_VIEW_TYPE_CUBE);
+	if ((NumColorAttachments == 0) && (Texture->GetViewType() == VK_IMAGE_VIEW_TYPE_CUBE))
+	{
+		InOutNumLayers = 6;
+
+		FVulkanTextureViewDesc ViewDesc;
+		ViewDesc.ViewType = VK_IMAGE_VIEW_TYPE_2D_ARRAY;
+		ViewDesc.AspectFlags = Texture->GetFullAspectMask();
+		ViewDesc.UEFormat = Texture->GetDesc().Format;
+		ViewDesc.Format = Texture->ViewFormat;
+		ViewDesc.FirstMip = MipIndex;
+		ViewDesc.NumMips = 1;
+		ViewDesc.ArraySliceIndex = 0;
+		ViewDesc.NumArraySlices = 6;
+		ViewDesc.bUseIdentitySwizzle = true;
+		return Texture->FindOrAddInternalView(ViewDesc);
+	}
+	else if ((Texture->GetViewType() == VK_IMAGE_VIEW_TYPE_2D) || (Texture->GetViewType() == VK_IMAGE_VIEW_TYPE_2D_ARRAY))
+	{
+		// depth attachments need a separate view to have no swizzle components, for validation correctness
+		FVulkanTextureViewDesc ViewDesc;
+		ViewDesc.ViewType = Texture->GetViewType();
+		ViewDesc.AspectFlags = Texture->GetFullAspectMask();
+		ViewDesc.UEFormat = Texture->GetDesc().Format;
+		ViewDesc.Format = Texture->ViewFormat;
+		ViewDesc.FirstMip = MipIndex;
+		ViewDesc.NumMips = 1;
+		ViewDesc.ArraySliceIndex = 0;
+		ViewDesc.NumArraySlices = Texture->GetNumberOfArrayLevels();
+		ViewDesc.bUseIdentitySwizzle = true;
+		return Texture->FindOrAddInternalView(ViewDesc);
+	}
+	else
+	{
+		return Texture->DefaultView;
+	}
+}
+
+FVulkanView* FVulkanFramebuffer::GetDepthStencilResolveTargetViewDesc(FVulkanTexture* ResolveTexture, uint32 MipIndex)
+{
+	// Resolve attachments only supported for 2d/2d array textures
+	if (ResolveTexture->GetViewType() == VK_IMAGE_VIEW_TYPE_2D || ResolveTexture->GetViewType() == VK_IMAGE_VIEW_TYPE_2D_ARRAY)
+	{
+		FVulkanTextureViewDesc ViewDesc;
+		ViewDesc.ViewType = ResolveTexture->GetViewType();
+		ViewDesc.AspectFlags = ResolveTexture->GetFullAspectMask();
+		ViewDesc.UEFormat = ResolveTexture->GetDesc().Format;
+		ViewDesc.Format = ResolveTexture->ViewFormat;
+		ViewDesc.FirstMip = MipIndex;
+		ViewDesc.NumMips = 1;
+		ViewDesc.ArraySliceIndex = 0;
+		ViewDesc.NumArraySlices = ResolveTexture->GetNumberOfArrayLevels();
+		ViewDesc.bUseIdentitySwizzle = true;
+		return ResolveTexture->FindOrAddInternalView(ViewDesc);
+	}
+
+	return nullptr;
+}
+
+FVulkanView* FVulkanFramebuffer::GetFragmentDensityAttachmentViewDesc(FVulkanTexture* Texture, uint32 MipIndex)
+{
+	ensure(Texture->GetViewType() == VK_IMAGE_VIEW_TYPE_2D || Texture->GetViewType() == VK_IMAGE_VIEW_TYPE_2D_ARRAY);
+
+	FVulkanTextureViewDesc ViewDesc;
+	ViewDesc.ViewType = Texture->GetViewType();
+	ViewDesc.AspectFlags = Texture->GetFullAspectMask();
+	ViewDesc.UEFormat = Texture->GetDesc().Format;
+	ViewDesc.Format = Texture->ViewFormat;
+	ViewDesc.FirstMip = MipIndex;
+	ViewDesc.NumMips = 1;
+	ViewDesc.ArraySliceIndex = 0;
+	ViewDesc.NumArraySlices = Texture->GetNumberOfArrayLevels();
+	ViewDesc.bUseIdentitySwizzle = true;
+	return Texture->FindOrAddInternalView(ViewDesc);
+}
+
+
+
+
 FVulkanFramebuffer::~FVulkanFramebuffer()
 {
 	ensure(Framebuffer == VK_NULL_HANDLE);
@@ -664,6 +678,12 @@ bool FVulkanFramebuffer::Matches(const FRHISetRenderTargetsInfo& InRTInfo) const
 	return true;
 }
 
+
+
+
+
+
+
 // Tear down and recreate swapchain and related resources.
 void FVulkanViewport::RecreateSwapchain(FVulkanCommandListContext& Context, FVulkanPlatformWindowContext& WindowContext)
 {
@@ -745,9 +765,17 @@ void FVulkanViewport::InitImages(FVulkanContextCommon& Context, TConstArrayView<
 			Swap(ImageSizeX, ImageSizeY);
 		}
 		BackBufferImages[Index] = (FVulkanTexture*)FVulkanDynamicRHI::Get().RHICreateTexture2DFromResource(PixelFormat, ImageSizeX, ImageSizeY, 1, 1, Images[Index], ETextureCreateFlags::RenderTargetable | ETextureCreateFlags::Presentable).GetReference();
-		const VkDescriptorType DescriptorType = VK_DESCRIPTOR_TYPE_SAMPLED_IMAGE;
-		TextureViews.Add((new FVulkanView(*Device, DescriptorType))->InitAsTextureView(
-			Images[Index], VK_IMAGE_VIEW_TYPE_2D, VK_IMAGE_ASPECT_COLOR_BIT, PixelFormat, UEToVkTextureFormat(PixelFormat, false), 0, 1, 0, 1, false));
+
+		FVulkanTextureViewDesc ViewDesc;
+		ViewDesc.ViewType = VK_IMAGE_VIEW_TYPE_2D;
+		ViewDesc.AspectFlags = VK_IMAGE_ASPECT_COLOR_BIT;
+		ViewDesc.UEFormat = PixelFormat;
+		ViewDesc.Format = UEToVkTextureFormat(PixelFormat, false);
+		ViewDesc.FirstMip = 0;
+		ViewDesc.NumMips = 1;
+		ViewDesc.ArraySliceIndex = 0;
+		ViewDesc.NumArraySlices = 1;
+		TextureViews.Add((new FVulkanView(*Device, VK_DESCRIPTOR_TYPE_SAMPLED_IMAGE))->InitAsTextureView(Images[Index], ViewDesc));
 
 		// Clear the swapchain to avoid a validation warning, and transition to PresentSrc
 		{
