RHI resource memory dump support for Vulkan RHI

diff --git a/Engine/Source/Runtime/VulkanRHI/Private/VulkanMemory.cpp b/Engine/Source/Runtime/VulkanRHI/Private/VulkanMemory.cpp
--- a/Engine/Source/Runtime/VulkanRHI/Private/VulkanMemory.cpp
+++ b/Engine/Source/Runtime/VulkanRHI/Private/VulkanMemory.cpp
@@ -3764,7 +3764,7 @@ namespace VulkanRHI
 	}
 
 	FVulkanAllocation::FVulkanAllocation()
-		: bHasOwnership(0)
+		: bHasOwnership(0), bTransient(false)
 	{
 		SetType(EVulkanAllocationEmpty);
 	}
diff --git a/Engine/Source/Runtime/VulkanRHI/Private/VulkanMemory.h b/Engine/Source/Runtime/VulkanRHI/Private/VulkanMemory.h
--- a/Engine/Source/Runtime/VulkanRHI/Private/VulkanMemory.h
+++ b/Engine/Source/Runtime/VulkanRHI/Private/VulkanMemory.h
@@ -490,8 +490,9 @@ namespace VulkanRHI
 		uint32 AllocationIndex = 0;
 		uint16 AllocatorIndex = 0;
 		EVulkanAllocationMetaType MetaType = EVulkanAllocationMetaUnknown;
-		uint8 Type : 7;
+		uint8 Type : 6;
 		uint8 bHasOwnership : 1;
+		uint8 bTransient : 1;
 
 		static_assert(EVulkanAllocationSize < 128, "Not enough bits to hold EVulkanAllocationType");
 		inline EVulkanAllocationType GetType() const
diff --git a/Engine/Source/Runtime/VulkanRHI/Private/VulkanResources.h b/Engine/Source/Runtime/VulkanRHI/Private/VulkanResources.h
--- a/Engine/Source/Runtime/VulkanRHI/Private/VulkanResources.h
+++ b/Engine/Source/Runtime/VulkanRHI/Private/VulkanResources.h
@@ -671,6 +671,11 @@ public:
 	void DestroySurface();
 	void InvalidateMappedMemory();
 	void* GetMappedPointer();
+		
+	// FRHIResource overrides 
+#if RHI_ENABLE_RESOURCE_INFO
+	bool GetResourceInfo(FRHIResourceInfo& OutResourceInfo) const override;
+#endif
 
 	/**
 	 * Returns how much memory is used by the surface
@@ -888,6 +893,20 @@ public:
 		LockCounter++;
 	}
 
+	// FRHIResource overrides
+#if RHI_ENABLE_RESOURCE_INFO
+	bool GetResourceInfo(FRHIResourceInfo& OutResourceInfo) const override
+	{
+		OutResourceInfo = FRHIResourceInfo{};
+		OutResourceInfo.Name = GetName();
+		OutResourceInfo.Type = GetType();
+		OutResourceInfo.VRamAllocation.AllocationSize = GetCurrentSize();
+		OutResourceInfo.IsTransient = CurrentBufferAlloc.Alloc.bTransient;
+
+		return true;
+	}
+#endif
+
 protected:
 
 	// Will return a new allocation that can be used for this buffer (proper memory type and size)
diff --git a/Engine/Source/Runtime/VulkanRHI/Private/VulkanTexture.cpp b/Engine/Source/Runtime/VulkanRHI/Private/VulkanTexture.cpp
--- a/Engine/Source/Runtime/VulkanRHI/Private/VulkanTexture.cpp
+++ b/Engine/Source/Runtime/VulkanRHI/Private/VulkanTexture.cpp
@@ -733,7 +733,18 @@ void* FVulkanTexture::GetMappedPointer()
 	return Allocation.GetMappedPointer(Device);
 }
 
+#if RHI_ENABLE_RESOURCE_INFO
+bool FVulkanTexture::GetResourceInfo(FRHIResourceInfo& OutResourceInfo) const
+{
+	OutResourceInfo = FRHIResourceInfo{};
+	OutResourceInfo.Name = GetName();
+	OutResourceInfo.Type = GetType();
+	OutResourceInfo.VRamAllocation.AllocationSize = GetMemorySize();
+	OutResourceInfo.IsTransient = Allocation.bTransient;
 
+	return true;
+}
+#endif
 
 VkDeviceMemory FVulkanTexture::GetAllocationHandle() const
 {
diff --git a/Engine/Source/Runtime/VulkanRHI/Private/VulkanTransientResourceAllocator.cpp b/Engine/Source/Runtime/VulkanRHI/Private/VulkanTransientResourceAllocator.cpp
--- a/Engine/Source/Runtime/VulkanRHI/Private/VulkanTransientResourceAllocator.cpp
+++ b/Engine/Source/Runtime/VulkanRHI/Private/VulkanTransientResourceAllocator.cpp
@@ -60,6 +60,7 @@ VulkanRHI::FVulkanAllocation FVulkanTransientHeap::GetVulkanAllocation(const FRH
 	TransientAlloc.HandleId = Heap->HandleID;
 	TransientAlloc.Offset += HeapAllocation.Offset;
 	TransientAlloc.Size = HeapAllocation.Size;
+	TransientAlloc.bTransient = true;
 	check((TransientAlloc.Offset + TransientAlloc.Size) <= Heap->InternalAllocation.Size);
 	return TransientAlloc;
 }
