VulkanRHI: Load dynamic rendering extensions and functions.  Change buffer markers to be enabled with UseGPUCrashBreadcrumbs.

#rnx

diff --git a/Engine/Source/Runtime/VulkanRHI/Private/VulkanCommandWrappers.h b/Engine/Source/Runtime/VulkanRHI/Private/VulkanCommandWrappers.h
--- a/Engine/Source/Runtime/VulkanRHI/Private/VulkanCommandWrappers.h
+++ b/Engine/Source/Runtime/VulkanRHI/Private/VulkanCommandWrappers.h
@@ -190,8 +190,11 @@ struct FWrapLayer
 	static void CmdBeginRenderPass2KHR(VkResult Result, VkCommandBuffer CommandBuffer, const VkRenderPassBeginInfo* RenderPassBegin, const VkSubpassBeginInfo* pSubpassBeginInfo) VULKAN_LAYER_BODY
 	static void CmdNextSubpass(VkResult Result, VkCommandBuffer CommandBuffer, VkSubpassContents Contents) VULKAN_LAYER_BODY
 	static void CmdEndRenderPass(VkResult Result, VkCommandBuffer CommandBuffer) VULKAN_LAYER_BODY
+	static void CmdBeginRenderingKHR(VkResult Result, VkCommandBuffer CommandBuffer, const VkRenderingInfo* RenderingInfo) VULKAN_LAYER_BODY
+	static void CmdEndRenderingKHR(VkResult Result, VkCommandBuffer CommandBuffer) VULKAN_LAYER_BODY
 	static void CmdExecuteCommands(VkResult Result, VkCommandBuffer CommandBuffer, uint32 CommandBufferCount, const VkCommandBuffer* pCommandBuffers) VULKAN_LAYER_BODY
 
+
 	static void QueuePresent(VkResult Result, VkQueue Queue, const VkPresentInfoKHR* PresentInfo) VULKAN_LAYER_BODY
 	static void GetSwapChainImagesKHR(VkResult Result, VkDevice Device, VkSwapchainKHR Swapchain, uint32_t* SwapchainImageCount, VkImage* SwapchainImages) VULKAN_LAYER_BODY
 	static void CreateSwapchainKHR(VkResult Result, VkDevice Device, const VkSwapchainCreateInfoKHR* CreateInfo, VkSwapchainKHR* Swapchain) VULKAN_LAYER_BODY
@@ -1382,6 +1385,20 @@ namespace VulkanRHI
 		FWrapLayer::CmdEndRenderPass(VK_SUCCESS, CommandBuffer);
 	}
 
+	static FORCEINLINE_DEBUGGABLE void vkCmdBeginRenderingKHR(VkCommandBuffer CommandBuffer, const VkRenderingInfo* RenderingInfo)
+	{
+		FWrapLayer::CmdBeginRenderingKHR(VK_RESULT_MAX_ENUM, CommandBuffer, RenderingInfo);
+		VULKANAPINAMESPACE::vkCmdBeginRenderingKHR(CommandBuffer, RenderingInfo);
+		FWrapLayer::CmdBeginRenderingKHR(VK_SUCCESS, CommandBuffer, RenderingInfo);
+	}
+
+	static FORCEINLINE_DEBUGGABLE void vkCmdEndRenderingKHR(VkCommandBuffer CommandBuffer)
+	{
+		FWrapLayer::CmdEndRenderingKHR(VK_RESULT_MAX_ENUM, CommandBuffer);
+		VULKANAPINAMESPACE::vkCmdEndRenderingKHR(CommandBuffer);
+		FWrapLayer::CmdEndRenderingKHR(VK_SUCCESS, CommandBuffer);
+	}
+
 	static FORCEINLINE_DEBUGGABLE void  vkCmdExecuteCommands(VkCommandBuffer CommandBuffer, uint32 CommandBufferCount, const VkCommandBuffer* pCommandBuffers)
 	{
 		FWrapLayer::CmdExecuteCommands(VK_RESULT_MAX_ENUM, CommandBuffer, CommandBufferCount, pCommandBuffers);
diff --git a/Engine/Source/Runtime/VulkanRHI/Private/VulkanDebug.cpp b/Engine/Source/Runtime/VulkanRHI/Private/VulkanDebug.cpp
--- a/Engine/Source/Runtime/VulkanRHI/Private/VulkanDebug.cpp
+++ b/Engine/Source/Runtime/VulkanRHI/Private/VulkanDebug.cpp
@@ -117,7 +117,8 @@ static VKAPI_ATTR VkBool32 VKAPI_CALL DebugUtilsCallback(VkDebugUtilsMessageSeve
 	}
 
 	if (!FCStringAnsi::Strcmp(CallbackData->pMessageIdName, "UNASSIGNED-CoreValidation-Shader-OutputNotConsumed") ||
-		!FCStringAnsi::Strcmp(CallbackData->pMessageIdName, "Undefined-Value-ShaderOutputNotConsumed"))
+		!FCStringAnsi::Strcmp(CallbackData->pMessageIdName, "Undefined-Value-ShaderOutputNotConsumed") ||
+		!FCStringAnsi::Strcmp(CallbackData->pMessageIdName, "Undefined-Value-ShaderOutputNotConsumed-DynamicRendering"))
 	{
 		// Warning: *** [Warning:Validation-1(UNASSIGNED-CoreValidation-Shader-OutputNotConsumed)] fragment shader writes to output location 0 with no matching attachment
 		return VK_FALSE;
@@ -2534,6 +2535,26 @@ void FWrapLayer::CmdNextSubpass(VkResult Result, VkCommandBuffer CommandBuffer,
 	}
 }
 
+void FWrapLayer::CmdBeginRenderingKHR(VkResult Result, VkCommandBuffer CommandBuffer, const VkRenderingInfo* RenderingInfo)
+{
+	if (Result == VK_RESULT_MAX_ENUM)
+	{
+#if VULKAN_ENABLE_DUMP_LAYER
+		CmdPrintfBegin(CommandBuffer, TEXT("----- vkCmdBeginRenderingKHR()"));
+#endif
+	}
+}
+
+void FWrapLayer::CmdEndRenderingKHR(VkResult Result, VkCommandBuffer CommandBuffer)
+{
+	if (Result == VK_RESULT_MAX_ENUM)
+	{
+#if VULKAN_ENABLE_DUMP_LAYER
+		CmdPrintfBegin(CommandBuffer, TEXT("----- vkCmdEndRenderingKHR()"));
+#endif
+	}
+}
+
 void FWrapLayer::CmdExecuteCommands(VkResult Result, VkCommandBuffer CommandBuffer, uint32 CommandBufferCount, const VkCommandBuffer* pCommandBuffers)
 {
 	if (Result == VK_RESULT_MAX_ENUM)
diff --git a/Engine/Source/Runtime/VulkanRHI/Private/VulkanDevice.h b/Engine/Source/Runtime/VulkanRHI/Private/VulkanDevice.h
--- a/Engine/Source/Runtime/VulkanRHI/Private/VulkanDevice.h
+++ b/Engine/Source/Runtime/VulkanRHI/Private/VulkanDevice.h
@@ -97,6 +97,7 @@ struct FOptionalVulkanDeviceExtensions
 			uint64 HasKHRMaintenance4 : 1;
 			uint64 HasKHRMaintenance5 : 1;
 			uint64 HasKHRSynchronization2 : 1;
+			uint64 HasKHRDynamicRendering : 1;
 			uint64 HasEXTSubgroupSizeControl : 1;
 			uint64 HasEXTPipelineCreationCacheControl : 1;
 		};
diff --git a/Engine/Source/Runtime/VulkanRHI/Private/VulkanExtensions.cpp b/Engine/Source/Runtime/VulkanRHI/Private/VulkanExtensions.cpp
--- a/Engine/Source/Runtime/VulkanRHI/Private/VulkanExtensions.cpp
+++ b/Engine/Source/Runtime/VulkanRHI/Private/VulkanExtensions.cpp
@@ -73,6 +73,15 @@ TAutoConsoleVariable<int32> GVulkanAllowSync2BarriersCVar(
 	ECVF_ReadOnly
 );
 
+TAutoConsoleVariable<int32> GVulkanAllowDynamicRenderingCVar(
+	TEXT("r.Vulkan.AllowDynamicRendering"),
+	0,
+	TEXT("Enables the use of VK_KHR_dynamic_rendering to remove the need for render passes. \n")
+	TEXT("0: Do not enable support for dynamic rendering (default)\n")
+	TEXT("1: Enable dynamic rendering"),
+	ECVF_ReadOnly
+);
+
 TAutoConsoleVariable<int32> GVulkanVariableRateShadingFormatCVar(
 	TEXT("r.Vulkan.VRSFormat"),
 	0,
@@ -1171,7 +1180,7 @@ public:
 	FVulkanAMDBufferMarkerExtension(FVulkanDevice* InDevice)
 		: FVulkanDeviceExtension(InDevice, VK_AMD_BUFFER_MARKER_EXTENSION_NAME, VULKAN_EXTENSION_ENABLED)
 	{
-		bEnabledInCode = bEnabledInCode && UE::RHI::UseGPUCrashDebugging() && UE::RHICore::AllowVendorDevice();
+		bEnabledInCode = bEnabledInCode && UE::RHI::UseGPUCrashBreadcrumbs() && UE::RHICore::AllowVendorDevice();
 	}
 
 	virtual void PostPhysicalDeviceFeatures(FOptionalVulkanDeviceExtensions& ExtensionFlags) override final
@@ -1804,6 +1813,127 @@ public:
 	VkPhysicalDeviceMutableDescriptorTypeFeaturesEXT MutableDescriptorTypeFeatures;
 };
 
+// ***** VK_KHR_dynamic_rendering
+class FVulkanKHRDynamicRenderingExtension : public FVulkanDeviceExtension
+{
+public:
+
+	FVulkanKHRDynamicRenderingExtension(FVulkanDevice* InDevice)
+		: FVulkanDeviceExtension(InDevice, VK_KHR_DYNAMIC_RENDERING_EXTENSION_NAME, VULKAN_EXTENSION_ENABLED, VK_API_VERSION_1_3)
+	{
+	}
+
+	virtual void PrePhysicalDeviceFeatures(VkPhysicalDeviceFeatures2KHR& PhysicalDeviceFeatures2) override final
+	{
+		ZeroVulkanStruct(DynamicRenderingFeatures, VK_STRUCTURE_TYPE_PHYSICAL_DEVICE_DYNAMIC_RENDERING_FEATURES);
+		AddToPNext(PhysicalDeviceFeatures2, DynamicRenderingFeatures);
+	}
+
+	virtual void PostPhysicalDeviceFeatures(FOptionalVulkanDeviceExtensions& ExtensionFlags) override final
+	{
+		bRequirementsPassed = (DynamicRenderingFeatures.dynamicRendering == VK_TRUE);
+
+		// We still want to load the extension if dynamic rendering isn't used (it's a dependency for Maintenance5).
+		if (GVulkanAllowDynamicRenderingCVar.GetValueOnAnyThread() != 0)
+		{
+			ExtensionFlags.HasKHRDynamicRendering = bRequirementsPassed;
+		}
+	}
+
+	virtual void PreCreateDevice(VkDeviceCreateInfo& DeviceCreateInfo) override final
+	{
+		if (bRequirementsPassed)
+		{
+			AddToPNext(DeviceCreateInfo, DynamicRenderingFeatures);
+		}
+	}
+
+	VkPhysicalDeviceDynamicRenderingFeaturesKHR DynamicRenderingFeatures;
+};
+
+// ***** VK_KHR_dynamic_rendering_local_read
+class FVulkanKHRDynamicRenderingLocalReadExtension : public FVulkanDeviceExtension
+{
+public:
+
+	FVulkanKHRDynamicRenderingLocalReadExtension(FVulkanDevice* InDevice)
+		: FVulkanDeviceExtension(InDevice, VK_KHR_DYNAMIC_RENDERING_LOCAL_READ_EXTENSION_NAME, VULKAN_EXTENSION_ENABLED, VK_API_VERSION_1_4)
+	{
+		bEnabledInCode = bEnabledInCode && (GVulkanAllowDynamicRenderingCVar.GetValueOnAnyThread() != 0);
+	}
+
+	virtual void PrePhysicalDeviceFeatures(VkPhysicalDeviceFeatures2KHR& PhysicalDeviceFeatures2) override final
+	{
+		ZeroVulkanStruct(DynamicRenderingLocalReadFeatures, VK_STRUCTURE_TYPE_PHYSICAL_DEVICE_DYNAMIC_RENDERING_LOCAL_READ_FEATURES);
+		AddToPNext(PhysicalDeviceFeatures2, DynamicRenderingLocalReadFeatures);
+	}
+
+	virtual void PostPhysicalDeviceFeatures(FOptionalVulkanDeviceExtensions& ExtensionFlags) override final
+	{
+		// Depends on the base dynamic rendering extension 
+		if (ExtensionFlags.HasKHRDynamicRendering)
+		{
+			bRequirementsPassed = (DynamicRenderingLocalReadFeatures.dynamicRenderingLocalRead == VK_TRUE);
+			ExtensionFlags.HasKHRDynamicRendering = bRequirementsPassed;
+		}
+		else
+		{
+			bRequirementsPassed = false;
+		}
+	}
+
+	virtual void PreCreateDevice(VkDeviceCreateInfo& DeviceCreateInfo) override final
+	{
+		if (bRequirementsPassed)
+		{
+			AddToPNext(DeviceCreateInfo, DynamicRenderingLocalReadFeatures);
+		}
+	}
+
+	VkPhysicalDeviceDynamicRenderingLocalReadFeatures DynamicRenderingLocalReadFeatures;
+};
+
+// ***** VK_EXT_dynamic_rendering_unused_attachments
+class FVulkanKHRDynamicRenderingUnusedAttachmentsExtension : public FVulkanDeviceExtension
+{
+public:
+
+	FVulkanKHRDynamicRenderingUnusedAttachmentsExtension(FVulkanDevice* InDevice)
+		: FVulkanDeviceExtension(InDevice, VK_EXT_DYNAMIC_RENDERING_UNUSED_ATTACHMENTS_EXTENSION_NAME, VULKAN_EXTENSION_ENABLED)
+	{
+		bEnabledInCode = bEnabledInCode && (GVulkanAllowDynamicRenderingCVar.GetValueOnAnyThread() != 0);
+	}
+
+	virtual void PrePhysicalDeviceFeatures(VkPhysicalDeviceFeatures2KHR& PhysicalDeviceFeatures2) override final
+	{
+		ZeroVulkanStruct(DynamicRenderingUnusedAttachmentsFeatures, VK_STRUCTURE_TYPE_PHYSICAL_DEVICE_DYNAMIC_RENDERING_UNUSED_ATTACHMENTS_FEATURES_EXT);
+		AddToPNext(PhysicalDeviceFeatures2, DynamicRenderingUnusedAttachmentsFeatures);
+	}
+
+	virtual void PostPhysicalDeviceFeatures(FOptionalVulkanDeviceExtensions& ExtensionFlags) override final
+	{
+		// Depends on the base dynamic rendering extension 
+		if (ExtensionFlags.HasKHRDynamicRendering)
+		{
+			bRequirementsPassed = (DynamicRenderingUnusedAttachmentsFeatures.dynamicRenderingUnusedAttachments == VK_TRUE);
+			ExtensionFlags.HasKHRDynamicRendering = bRequirementsPassed;
+		}
+		else
+		{
+			bRequirementsPassed = false;
+		}
+	}
+
+	virtual void PreCreateDevice(VkDeviceCreateInfo& DeviceCreateInfo) override final
+	{
+		if (bRequirementsPassed)
+		{
+			AddToPNext(DeviceCreateInfo, DynamicRenderingUnusedAttachmentsFeatures);
+		}
+	}
+
+	VkPhysicalDeviceDynamicRenderingUnusedAttachmentsFeaturesEXT DynamicRenderingUnusedAttachmentsFeatures;
+};
 
 template <typename ExtensionType>
 static void FlagExtensionSupport(const TArray<VkExtensionProperties>& ExtensionProperties, TArray<TUniquePtr<ExtensionType>>& UEExtensions, uint32 ApiVersion, const TCHAR* ExtensionTypeName)
@@ -1857,7 +1987,6 @@ FVulkanDeviceExtensionArray FVulkanDeviceExtension::GetUESupportedDeviceExtensio
 	ADD_SIMPLE_EXTENSION(VK_EXT_LOAD_STORE_OP_NONE_EXTENSION_NAME,       VULKAN_EXTENSION_ENABLED,             VULKAN_EXTENSION_NOT_PROMOTED, DEVICE_EXT_FLAG_SETTER(HasEXTLoadStoreOpNone));
 	ADD_SIMPLE_EXTENSION(VK_QCOM_RENDER_PASS_SHADER_RESOLVE_EXTENSION_NAME, VULKAN_SUPPORTS_QCOM_RENDERPASS_SHADER_RESOLVE, VULKAN_EXTENSION_NOT_PROMOTED, DEVICE_EXT_FLAG_SETTER(HasQcomRenderPassShaderResolve));
 	ADD_SIMPLE_EXTENSION(VK_EXT_TOOLING_INFO_EXTENSION_NAME, VULKAN_HAS_DEBUGGING_ENABLED || VULKAN_ENABLE_DRAW_MARKERS, VK_API_VERSION_1_3,  DEVICE_EXT_FLAG_SETTER(HasEXTToolingInfo));
-	ADD_SIMPLE_EXTENSION(VK_KHR_DYNAMIC_RENDERING_EXTENSION_NAME,        VULKAN_EXTENSION_ENABLED,             VK_API_VERSION_1_3,            nullptr); // Only enabled because maintenance5 depends on it
 	ADD_SIMPLE_EXTENSION(VK_KHR_PIPELINE_LIBRARY_EXTENSION_NAME,         VULKAN_EXTENSION_ENABLED,             VULKAN_EXTENSION_NOT_PROMOTED, DEVICE_EXT_FLAG_SETTER(HasKHRPipelineLibrary));
 
 	// Externally activated extensions (supported by the engine, but enabled externally by plugin or other) :
@@ -1899,6 +2028,9 @@ FVulkanDeviceExtensionArray FVulkanDeviceExtension::GetUESupportedDeviceExtensio
 	ADD_CUSTOM_EXTENSION(FVulkanEXTImageCompressionControlExtension);
 	ADD_CUSTOM_EXTENSION(FVulkanKHRGraphicsPipelineLibraryExtension);
 	ADD_CUSTOM_EXTENSION(FVulkanEXTMutableDescriptorTypeExtension);
+	ADD_CUSTOM_EXTENSION(FVulkanKHRDynamicRenderingExtension);			// must be kept BEFORE DynamicRenderingLocalRead!
+	ADD_CUSTOM_EXTENSION(FVulkanKHRDynamicRenderingLocalReadExtension);	// must be kept AFTER DynamicRendering!
+	ADD_CUSTOM_EXTENSION(FVulkanKHRDynamicRenderingUnusedAttachmentsExtension);	// must be kept AFTER DynamicRendering!
 
 	// Needed for Raytracing
 	ADD_CUSTOM_EXTENSION(FVulkanKHRBufferDeviceAddressExtension);
diff --git a/Engine/Source/Runtime/VulkanRHI/Private/VulkanLoader.h b/Engine/Source/Runtime/VulkanRHI/Private/VulkanLoader.h
--- a/Engine/Source/Runtime/VulkanRHI/Private/VulkanLoader.h
+++ b/Engine/Source/Runtime/VulkanRHI/Private/VulkanLoader.h
@@ -204,6 +204,8 @@
 	ENUM_VK_ENTRYPOINTS_RAYTRACING(EnumMacro) \
 	EnumMacro(PFN_vkCreateRenderPass2KHR, vkCreateRenderPass2KHR) \
 	EnumMacro(PFN_vkCmdBeginRenderPass2KHR, vkCmdBeginRenderPass2KHR) \
+	EnumMacro(PFN_vkCmdBeginRenderingKHR, vkCmdBeginRenderingKHR) \
+	EnumMacro(PFN_vkCmdEndRenderingKHR, vkCmdEndRenderingKHR) \
 	EnumMacro(PFN_vkGetDeviceImageMemoryRequirementsKHR, vkGetDeviceImageMemoryRequirementsKHR) \
 	EnumMacro(PFN_vkGetDeviceBufferMemoryRequirementsKHR, vkGetDeviceBufferMemoryRequirementsKHR) \
 	EnumMacro(PFN_vkResetQueryPoolEXT, vkResetQueryPoolEXT) \
