VulkanRHI: Add support for dynamic rendering on desktop (disabled by default).

#rnx

diff --git a/Engine/Source/Runtime/VulkanRHI/Private/VulkanCommandBuffer.cpp b/Engine/Source/Runtime/VulkanRHI/Private/VulkanCommandBuffer.cpp
--- a/Engine/Source/Runtime/VulkanRHI/Private/VulkanCommandBuffer.cpp
+++ b/Engine/Source/Runtime/VulkanRHI/Private/VulkanCommandBuffer.cpp
@@ -35,7 +35,7 @@ void FVulkanCommandBuffer::AllocMemory()
 
 	VkCommandBufferAllocateInfo CreateCmdBufInfo;
 	ZeroVulkanStruct(CreateCmdBufInfo, VK_STRUCTURE_TYPE_COMMAND_BUFFER_ALLOCATE_INFO);
-	CreateCmdBufInfo.level = (GetCommandBufferType() == EVulkanCommandBufferType::Primary) ? VK_COMMAND_BUFFER_LEVEL_PRIMARY : VK_COMMAND_BUFFER_LEVEL_SECONDARY;
+	CreateCmdBufInfo.level = (GetCommandBufferType() == EVulkanCommandBufferType::Secondary) ? VK_COMMAND_BUFFER_LEVEL_SECONDARY : VK_COMMAND_BUFFER_LEVEL_PRIMARY;
 	CreateCmdBufInfo.commandBufferCount = 1;
 	CreateCmdBufInfo.commandPool = CommandBufferPool.GetHandle();
 
@@ -110,7 +110,24 @@ void FVulkanCommandBuffer::BeginRenderPass(const FVulkanBeginRenderPassInfo& Beg
 	State = EState::IsInsideRenderPass;
 }
 
-void FVulkanCommandBuffer::End(FVulkanQueryPool* QueryPool)
+void FVulkanCommandBuffer::BeginDynamicRendering(const VkRenderingInfo& RenderingInfo)
+{
+	check(GetCommandBufferType() != EVulkanCommandBufferType::Secondary);
+
+	VulkanRHI::vkCmdBeginRenderingKHR(CommandBufferHandle, &RenderingInfo);
+
+	LastDynamicRenderingFlags = RenderingInfo.flags;
+	State = EState::IsInsideRenderPass;
+}
+
+void FVulkanCommandBuffer::EndDynamicRendering()
+{
+	check(GetCommandBufferType() != EVulkanCommandBufferType::Secondary);
+	VulkanRHI::vkCmdEndRenderingKHR(CommandBufferHandle);
+	State = EState::IsInsideBegin;
+}
+
+void FVulkanCommandBuffer::End(FVulkanQueryPool* OptionalQueryPool)
 {
 	checkf(IsOutsideRenderPass() || (GetCommandBufferType() == EVulkanCommandBufferType::Secondary), 
 		TEXT("Can't End as we're inside a render pass! CmdBuffer 0x%p State=%d"), CommandBufferHandle, (int32)State);
@@ -122,17 +139,17 @@ void FVulkanCommandBuffer::End(FVulkanQueryPool* QueryPool)
 	}
 
 #if RHI_NEW_GPU_PROFILER
-	if (QueryPool)
+	if (OptionalQueryPool)
 	{
 		auto& Event = EmplaceProfilerEvent<UE::RHI::GPUProfiler::FEvent::FEndWork>(0);
-		const uint32 IndexInPool = QueryPool->ReserveQuery(&Event.GPUTimestampBOP);
-		VulkanRHI::vkCmdWriteTimestamp(GetHandle(), VK_PIPELINE_STAGE_BOTTOM_OF_PIPE_BIT, QueryPool->GetHandle(), IndexInPool);
+		const uint32 IndexInPool = OptionalQueryPool->ReserveQuery(&Event.GPUTimestampBOP);
+		VulkanRHI::vkCmdWriteTimestamp(GetHandle(), VK_PIPELINE_STAGE_BOTTOM_OF_PIPE_BIT, OptionalQueryPool->GetHandle(), IndexInPool);
 	}
 #else
-	if (QueryPool)
+	if (OptionalQueryPool)
 	{
-		const uint32 IndexInPool = QueryPool->ReserveQuery(&EndTimestamp);
-		VulkanRHI::vkCmdWriteTimestamp(GetHandle(), VK_PIPELINE_STAGE_BOTTOM_OF_PIPE_BIT, QueryPool->GetHandle(), IndexInPool);
+		const uint32 IndexInPool = OptionalQueryPool->ReserveQuery(&EndTimestamp);
+		VulkanRHI::vkCmdWriteTimestamp(GetHandle(), VK_PIPELINE_STAGE_BOTTOM_OF_PIPE_BIT, OptionalQueryPool->GetHandle(), IndexInPool);
 	}
 #endif
 
@@ -140,9 +157,9 @@ void FVulkanCommandBuffer::End(FVulkanQueryPool* QueryPool)
 	State = EState::HasEnded;
 }
 
-void FVulkanCommandBuffer::Begin(FVulkanQueryPool* QueryPool, VkRenderPass RenderPassHandle)
+void FVulkanCommandBuffer::Begin(FVulkanQueryPool* OptionalQueryPool, VkRenderPass RenderPassHandle)
 {
-	checkf((GetCommandBufferType() == EVulkanCommandBufferType::Primary) || (RenderPassHandle != VK_NULL_HANDLE),
+	checkf((GetCommandBufferType() != EVulkanCommandBufferType::Secondary) || (RenderPassHandle != VK_NULL_HANDLE),
 		TEXT("Secondary command buffers require the render pass handle!"));
 
 	{
@@ -184,19 +201,19 @@ void FVulkanCommandBuffer::Begin(FVulkanQueryPool* QueryPool, VkRenderPass Rende
 	bNeedsDynamicStateSet = true;
 
 #if RHI_NEW_GPU_PROFILER
-	if (QueryPool)
+	if (OptionalQueryPool)
 	{
 		auto& Event = EmplaceProfilerEvent<UE::RHI::GPUProfiler::FEvent::FBeginWork>(0, UINT64_MAX);
-		const uint32 IndexInPool = QueryPool->ReserveQuery(&Event.GPUTimestampTOP);
-		VulkanRHI::vkCmdWriteTimestamp(GetHandle(), VK_PIPELINE_STAGE_TOP_OF_PIPE_BIT, QueryPool->GetHandle(), IndexInPool);
+		const uint32 IndexInPool = OptionalQueryPool->ReserveQuery(&Event.GPUTimestampTOP);
+		VulkanRHI::vkCmdWriteTimestamp(GetHandle(), VK_PIPELINE_STAGE_TOP_OF_PIPE_BIT, OptionalQueryPool->GetHandle(), IndexInPool);
 	}
 #else
-	if (QueryPool)
+	if (OptionalQueryPool)
 	{
 		StartTimestamp = 0;
 		EndTimestamp = 0;
-		const uint32 IndexInPool = QueryPool->ReserveQuery(&StartTimestamp);
-		VulkanRHI::vkCmdWriteTimestamp(GetHandle(), VK_PIPELINE_STAGE_TOP_OF_PIPE_BIT, QueryPool->GetHandle(), IndexInPool);
+		const uint32 IndexInPool = OptionalQueryPool->ReserveQuery(&StartTimestamp);
+		VulkanRHI::vkCmdWriteTimestamp(GetHandle(), VK_PIPELINE_STAGE_TOP_OF_PIPE_BIT, OptionalQueryPool->GetHandle(), IndexInPool);
 	}
 #endif
 }
diff --git a/Engine/Source/Runtime/VulkanRHI/Private/VulkanCommandBuffer.h b/Engine/Source/Runtime/VulkanRHI/Private/VulkanCommandBuffer.h
--- a/Engine/Source/Runtime/VulkanRHI/Private/VulkanCommandBuffer.h
+++ b/Engine/Source/Runtime/VulkanRHI/Private/VulkanCommandBuffer.h
@@ -18,6 +18,7 @@ class FVulkanQueue;
 class FVulkanDescriptorPoolSetContainer;
 class FVulkanQueryPool;
 struct FVulkanBeginRenderPassInfo;
+struct FVulkanDynamicRenderingInfo;
 
 namespace VulkanRHI
 {
@@ -28,7 +29,15 @@ namespace VulkanRHI
 enum class EVulkanCommandBufferType : uint8
 {
 	Primary,
+
+	// Secondary command buffer with an inherited render pass that will get executed from a primary cmd buffer
+	// Used for parallel contexts without dynamic rendering
 	Secondary,
+
+	// Primary command buffer that suspend/resume the current renderpass 
+	// Used for parallel contexts with dynamic rendering
+	Parallel,
+
 	Count
 };
 
@@ -81,8 +90,8 @@ public:
 
 	void SetSubmitted();
 
-	void Begin(FVulkanQueryPool* QueryPool, VkRenderPass RenderPassHandle);
-	void End(FVulkanQueryPool* QueryPool);
+	void Begin(FVulkanQueryPool* OptionalQueryPool, VkRenderPass RenderPassHandle);
+	void End(FVulkanQueryPool* OptionalQueryPool);
 
 	enum class EState : uint8
 	{
@@ -112,6 +121,13 @@ public:
 	void BeginRenderPass(const FVulkanBeginRenderPassInfo& BeginRenderPassInfo, const VkClearValue* AttachmentClearValues);
 	void EndRenderPass();
 
+	// Begin/End calls for when using VK_KHR_dynamic_rendering
+	void BeginDynamicRendering(const VkRenderingInfo& RenderingInfo);
+	void EndDynamicRendering();
+
+	// Flags used with the last dynamic rendering pass
+	VkRenderingFlags LastDynamicRenderingFlags = 0;
+
 	// Split barriers
 	void BeginSplitBarrier(VkEvent BarrierEvent, const VkDependencyInfo& DependencyInfo);
 	void EndSplitBarrier(VkEvent BarrierEvent, const VkDependencyInfo& DependencyInfo);
diff --git a/Engine/Source/Runtime/VulkanRHI/Private/VulkanContext.cpp b/Engine/Source/Runtime/VulkanRHI/Private/VulkanContext.cpp
--- a/Engine/Source/Runtime/VulkanRHI/Private/VulkanContext.cpp
+++ b/Engine/Source/Runtime/VulkanRHI/Private/VulkanContext.cpp
@@ -41,24 +41,47 @@ void FVulkanContextCommon::EndPayload()
 		if (Payload->CommandBuffers.Num() > 0)
 		{
 			FVulkanCommandBuffer* CommandBuffer = Payload->CommandBuffers.Last();
-			checkSlow(!CommandBuffer->IsSubmitted() && CommandBuffer->HasBegun());
-
-			const bool bIsPrimaryCommandBuffer = (CommandBuffer->GetCommandBufferType() == EVulkanCommandBufferType::Primary);
-			if (!CommandBuffer->IsOutsideRenderPass() && bIsPrimaryCommandBuffer)
+			if (!CommandBuffer->HasEnded())
 			{
-				UE_LOG(LogVulkanRHI, Warning, TEXT("Forcing EndRenderPass() for submission"));
-				CommandBuffer->EndRenderPass();
+				checkSlow(!CommandBuffer->IsSubmitted() && CommandBuffer->HasBegun());
+
+				const bool bIsPrimaryCommandBuffer = (CommandBuffer->GetCommandBufferType() != EVulkanCommandBufferType::Secondary);
+				if (!CommandBuffer->IsOutsideRenderPass() && bIsPrimaryCommandBuffer)
+				{
+					if (Device.GetOptionalExtensions().HasKHRDynamicRendering)
+					{
+						UE_LOG(LogVulkanRHI, Warning, TEXT("Forcing EndDynamicRendering() for submission"));
+						CommandBuffer->EndDynamicRendering();
+					}
+					else
+					{
+						UE_LOG(LogVulkanRHI, Warning, TEXT("Forcing EndRenderPass() for submission"));
+						CommandBuffer->EndRenderPass();
+					}
+				}
+
+				// Only record begin/end timestamps on primary command buffers
+				FVulkanQueryPool* OptionalTimestampQueryPool =
+					GRHIGlobals.SupportsTimestampRenderQueries && bIsPrimaryCommandBuffer ?
+					GetCurrentTimestampQueryPool(*Payload) : nullptr;
+				CommandBuffer->End(OptionalTimestampQueryPool);
 			}
-
-			// Only record begin/end timestamps on primary command buffers
-			FVulkanQueryPool* TimestampQueryPool = 
-				GRHIGlobals.SupportsTimestampRenderQueries && bIsPrimaryCommandBuffer ?
-				GetCurrentTimestampQueryPool(*Payload) : nullptr;
-			CommandBuffer->End(TimestampQueryPool);
 		}
 	}
 }
 
+void FVulkanContextCommon::AppendParallelRenderPayload(FVulkanPayload* ParallelRenderingPayload)
+{
+	check(Device.GetOptionalExtensions().HasKHRDynamicRendering);
+	check(ParallelRenderingPayload->CommandBuffers.IsEmpty() ||
+		(ParallelRenderingPayload->CommandBuffers[0]->GetCommandBufferType() == EVulkanCommandBufferType::Parallel));
+
+	EndPayload();
+	Payloads.Add(ParallelRenderingPayload);
+	PrepareNewCommandBuffer(*ParallelRenderingPayload);
+	CurrentPhase = EPhase::Execute;
+}
+
 // Complete recording of the current command list set, and appends the resulting
 // payloads to the given array. Resets the context so new commands can be recorded.
 void FVulkanContextCommon::Finalize(TArray<FVulkanPayload*>& OutPayloads)
@@ -176,11 +199,11 @@ void FVulkanContextCommon::PrepareNewCommandBuffer(FVulkanPayload& Payload)
 	Payload.CommandBuffers.Add(NewCommandBuffer);
 
 	// Only record begin/end timestamps on primary command buffers
-	const bool bIsPrimaryCommandBuffer = (NewCommandBuffer->GetCommandBufferType() == EVulkanCommandBufferType::Primary);
-	FVulkanQueryPool* TimestampQueryPool = GRHIGlobals.SupportsTimestampRenderQueries && bIsPrimaryCommandBuffer ?
+	const bool bIsPrimaryCommandBuffer = (NewCommandBuffer->GetCommandBufferType() != EVulkanCommandBufferType::Secondary);
+	FVulkanQueryPool* OptionalTimestampQueryPool = GRHIGlobals.SupportsTimestampRenderQueries && bIsPrimaryCommandBuffer ?
 		GetCurrentTimestampQueryPool(Payload) : nullptr;
 	VkRenderPass RenderPassHandle = GetParallelRenderPassInfo() ? GetParallelRenderPassInfo()->RenderPassHandle : VK_NULL_HANDLE;
-	NewCommandBuffer->Begin(TimestampQueryPool, RenderPassHandle);
+	NewCommandBuffer->Begin(OptionalTimestampQueryPool, RenderPassHandle);
 }
 
 
@@ -206,7 +229,8 @@ FVulkanCommandListContext::FVulkanCommandListContext(FVulkanDevice& InDevice, ER
 }
 
 FVulkanCommandListContext::FVulkanCommandListContext(FVulkanDevice& InDevice, FVulkanCommandListContext* InImmediate, FVulkanParallelRenderPassInfo* InParallelRenderPassInfo)
-	: FVulkanContextCommon(InDevice, *InDevice.GetQueue(ERHIPipeline::Graphics), EVulkanCommandBufferType::Secondary)
+	: FVulkanContextCommon(InDevice, *InDevice.GetQueue(ERHIPipeline::Graphics), 
+		InDevice.GetOptionalExtensions().HasKHRDynamicRendering ? EVulkanCommandBufferType::Parallel : EVulkanCommandBufferType::Secondary)
 	, Immediate(InImmediate)
 	, RHIPipeline(ERHIPipeline::Graphics)
 	, bSupportsBreadcrumbs(false)
@@ -250,65 +274,337 @@ FVulkanCommandListContext::~FVulkanCommandListContext()
 	ReleasePendingState();
 }
 
-void FVulkanCommandListContext::RHIBeginParallelRenderPass(TSharedPtr<FRHIParallelRenderPassInfo> InInfo, const TCHAR* InName)
+void FVulkanCommandListContext::FillDynamicRenderingInfo(const FRHIRenderPassInfo& InRenderPassInfo, FVulkanDynamicRenderingInfo& OutRenderingInfo)
 {
-	checkf(CurrentParallelRenderPassInfo == nullptr, TEXT("There is already a parallel render pass in progress!"));
-	CurrentParallelRenderPassInfo = new FVulkanParallelRenderPassInfo();
+	VkRenderingInfo& RenderingInfo = OutRenderingInfo.RenderingInfo;
+	ZeroVulkanStruct(RenderingInfo, VK_STRUCTURE_TYPE_RENDERING_INFO);
 
-	RHIBeginRenderPass(*InInfo.Get(), InName);
+	TStaticArray<VkRenderingAttachmentInfo, MaxSimultaneousRenderTargets>& ColorAttachmentInfos = OutRenderingInfo.ColorAttachmentInfos;
+	VkRenderingAttachmentInfo& DepthAttachmentInfo = OutRenderingInfo.DepthAttachmentInfo;
+	VkRenderingAttachmentInfo& StencilAttachmentInfo = OutRenderingInfo.StencilAttachmentInfo;
+	VkRenderingFragmentShadingRateAttachmentInfoKHR& FragmentShadingRateAttachmentInfo = OutRenderingInfo.FragmentShadingRateAttachmentInfo;
 
-	CurrentParallelRenderPassInfo->RenderPassHandle = CurrentRenderPass->GetHandle();
-	InInfo->RHIPlatformData = CurrentParallelRenderPassInfo;
-}
+	uint32 NumSamples = 0;
+	uint32 MipIndex = 0;
+	uint32 NumLayers = 0;
 
-void FVulkanCommandListContext::RHIEndParallelRenderPass()
-{
-	if (CurrentParallelRenderPassInfo->SecondaryPayloads.Num())
+	// CustomResolveSubpass can have targets with a different NumSamples
+	// Also with a CustomResolveSubpass last color attachment is a resolve target
+	const bool bHasCustomResolveSubpass = (InRenderPassInfo.SubpassHint == ESubpassHint::CustomResolveSubpass);
+
+	const int32 NumColorRenderTargets = InRenderPassInfo.GetNumColorRenderTargets();
+	if (NumColorRenderTargets)
 	{
-		FVulkanPayload& ParentPayload = GetPayload(EPhase::Execute);
-		FVulkanCommandBuffer& ParentCommandBuffer = GetCommandBuffer();
+		RenderingInfo.colorAttachmentCount = NumColorRenderTargets;
+		RenderingInfo.pColorAttachments = ColorAttachmentInfos.GetData();
+
+		for (int32 Index = 0; Index < NumColorRenderTargets; ++Index)
+		{
+			const FRHIRenderPassInfo::FColorEntry& ColorEntry = InRenderPassInfo.ColorRenderTargets[Index];
+
+			VkRenderingAttachmentInfo& AttachmentInfo = ColorAttachmentInfos[Index];
+			ZeroVulkanStruct(AttachmentInfo, VK_STRUCTURE_TYPE_RENDERING_ATTACHMENT_INFO);
+
+			FVulkanTexture* Texture = ResourceCast(ColorEntry.RenderTarget);
+			check(Texture);
+			const FRHITextureDesc& TextureDesc = Texture->GetDesc();
+
+			VkExtent2D& Extent = RenderingInfo.renderArea.extent;
+			const uint32 ExtentWidth = FMath::Max(1, TextureDesc.Extent.X >> ColorEntry.MipIndex);
+			ensure((Extent.width == 0) || (Extent.width == ExtentWidth));
+			Extent.width = ExtentWidth;
+			const uint32 ExtentHeight = FMath::Max(1, TextureDesc.Extent.Y >> ColorEntry.MipIndex);
+			ensure((Extent.height == 0) || (Extent.height == ExtentHeight));
+			Extent.height = ExtentHeight;
+
+			NumLayers = TextureDesc.Depth;
+
+			AttachmentInfo.imageLayout = VK_IMAGE_LAYOUT_COLOR_ATTACHMENT_OPTIMAL;
+			AttachmentInfo.loadOp = RenderTargetLoadActionToVulkan(GetLoadAction(ColorEntry.Action));
+			AttachmentInfo.storeOp = RenderTargetStoreActionToVulkan(GetStoreAction(ColorEntry.Action));
+			ensure(!EnumHasAnyFlags(TextureDesc.Flags, TexCreate_Memoryless) || (AttachmentInfo.storeOp == VK_ATTACHMENT_STORE_OP_DONT_CARE));
+
+			const FLinearColor ClearColor = Texture->HasClearValue() ? Texture->GetClearColor() : FLinearColor::Black;
+			AttachmentInfo.clearValue.color.float32[0] = ClearColor.R;
+			AttachmentInfo.clearValue.color.float32[1] = ClearColor.G;
+			AttachmentInfo.clearValue.color.float32[2] = ClearColor.B;
+			AttachmentInfo.clearValue.color.float32[3] = ClearColor.A;
+
+			MipIndex = ColorEntry.MipIndex;
+			FVulkanView* View = FVulkanFramebuffer::GetColorRenderTargetViewDesc(Texture, MipIndex, ColorEntry.ArraySlice, InRenderPassInfo.MultiViewCount, NumLayers);
+			AttachmentInfo.imageView = View->GetTextureView().View;
+
+			ensure(!NumSamples || (NumSamples == ColorEntry.RenderTarget->GetNumSamples()) || bHasCustomResolveSubpass);
+			NumSamples = ColorEntry.RenderTarget->GetNumSamples();
+
+			const bool bCustomResolveAttachment = (Index == (NumColorRenderTargets - 1)) && bHasCustomResolveSubpass;
+			const VkSampleCountFlagBits Samples = bCustomResolveAttachment ? VK_SAMPLE_COUNT_1_BIT : static_cast<VkSampleCountFlagBits>(NumSamples);
+			if ((Samples > VK_SAMPLE_COUNT_1_BIT) && ColorEntry.ResolveTarget)
+			{
+				if (FVulkanView* ResolveView = FVulkanFramebuffer::GetColorResolveTargetViewDesc(Texture, MipIndex, ColorEntry.ArraySlice))
+				{
+					AttachmentInfo.resolveMode = VK_RESOLVE_MODE_SAMPLE_ZERO_BIT;
+					AttachmentInfo.resolveImageLayout = VK_IMAGE_LAYOUT_COLOR_ATTACHMENT_OPTIMAL;
+					AttachmentInfo.resolveImageView = ResolveView->GetTextureView().View;
+				}
+			}
+		}
+	}
 
-		TArray<VkCommandBuffer> CommandBufferHandles;
-		CommandBufferHandles.Reserve(CurrentParallelRenderPassInfo->SecondaryPayloads.Num()); // should be 1 cmdbuffer per payload
-		for (FVulkanPayload* Payload : CurrentParallelRenderPassInfo->SecondaryPayloads)
+	if (InRenderPassInfo.DepthStencilRenderTarget.DepthStencilTarget)
+	{
+		VkImageLayout CurrentDepthLayout = VK_IMAGE_LAYOUT_UNDEFINED;
+		VkImageLayout CurrentStencilLayout = VK_IMAGE_LAYOUT_UNDEFINED;
+		ExtractDepthStencilLayouts(InRenderPassInfo, CurrentDepthLayout, CurrentStencilLayout);
+
+		FVulkanTexture* Texture = ResourceCast(InRenderPassInfo.DepthStencilRenderTarget.DepthStencilTarget);
+		check(Texture);
+		const FRHITextureDesc& TextureDesc = Texture->GetDesc();
+		check(IsDepthOrStencilFormat(TextureDesc.Format));
+
+		VkExtent2D& Extent = RenderingInfo.renderArea.extent;
+		Extent.width = (Extent.width == 0) ? TextureDesc.Extent.X : FMath::Min<int32>(Extent.width, TextureDesc.Extent.X);
+		Extent.height = (Extent.height == 0) ? TextureDesc.Extent.Y : FMath::Min<int32>(Extent.height, TextureDesc.Extent.Y);
+		NumLayers = (NumLayers == 0) ? Texture->GetNumberOfArrayLevels() : NumLayers;
+
+		FVulkanView* DepthView = nullptr;
+		bool bIsDepthStore = false;
+		if (InRenderPassInfo.DepthStencilRenderTarget.ExclusiveDepthStencil.IsUsingDepth())
 		{
-			check(Payload->SignalSemaphores.Num() == 0);
-			check(Payload->WaitSemaphores.Num() == 0);
+			ZeroVulkanStruct(DepthAttachmentInfo, VK_STRUCTURE_TYPE_RENDERING_ATTACHMENT_INFO);
+			RenderingInfo.pDepthAttachment = &DepthAttachmentInfo;
+			Texture->GetDepthStencilClearValue(DepthAttachmentInfo.clearValue.depthStencil.depth, DepthAttachmentInfo.clearValue.depthStencil.stencil);
+
+			// We can't have the final layout be UNDEFINED, but it's possible that we get here from a transient texture
+			// where the stencil was never used yet.  We can set the layout to whatever we want, the next transition will
+			// happen from UNDEFINED anyhow.
+			if (CurrentDepthLayout != VK_IMAGE_LAYOUT_UNDEFINED)
+			{
+				DepthAttachmentInfo.imageLayout = CurrentDepthLayout;
+				DepthAttachmentInfo.loadOp = RenderTargetLoadActionToVulkan(GetLoadAction(GetDepthActions(InRenderPassInfo.DepthStencilRenderTarget.Action)));
+				DepthAttachmentInfo.storeOp = RenderTargetStoreActionToVulkan(GetStoreAction(GetDepthActions(InRenderPassInfo.DepthStencilRenderTarget.Action)));
+				ensure(!EnumHasAnyFlags(TextureDesc.Flags, TexCreate_Memoryless) || (DepthAttachmentInfo.storeOp == VK_ATTACHMENT_STORE_OP_DONT_CARE));
 
-			for (int32 QueryPoolTypeIndex = 0; QueryPoolTypeIndex < (int32)EVulkanQueryPoolType::Count; ++QueryPoolTypeIndex)
+				bIsDepthStore = (DepthAttachmentInfo.storeOp == VK_ATTACHMENT_STORE_OP_STORE);
+			}
+			else
 			{
-				ParentPayload.QueryPools[QueryPoolTypeIndex].Append(MoveTemp(Payload->QueryPools[QueryPoolTypeIndex]));
+				// Unused image aspects with undefined layout should just remain untouched
+				DepthAttachmentInfo.imageLayout = VK_IMAGE_LAYOUT_ATTACHMENT_OPTIMAL;
+				DepthAttachmentInfo.loadOp = VK_ATTACHMENT_LOAD_OP_NONE;
+				DepthAttachmentInfo.storeOp = VK_ATTACHMENT_STORE_OP_NONE;
 			}
 
-			ParentPayload.DescriptorPoolSetContainers.Append(MoveTemp(Payload->DescriptorPoolSetContainers));
+			DepthView = FVulkanFramebuffer::GetDepthStencilTargetViewDesc(Texture, NumColorRenderTargets, MipIndex, NumLayers);
+			DepthAttachmentInfo.imageView = DepthView->GetTextureView().View;
+		}
+
+		bool bIsStencilStore = false;
+		if (InRenderPassInfo.DepthStencilRenderTarget.ExclusiveDepthStencil.IsUsingStencil() && IsStencilFormat(TextureDesc.Format))
+		{
+			ZeroVulkanStruct(StencilAttachmentInfo, VK_STRUCTURE_TYPE_RENDERING_ATTACHMENT_INFO);
+			RenderingInfo.pStencilAttachment = &StencilAttachmentInfo;
+			Texture->GetDepthStencilClearValue(StencilAttachmentInfo.clearValue.depthStencil.depth, StencilAttachmentInfo.clearValue.depthStencil.stencil);
 
-			for (FVulkanCommandBuffer* SecondaryCommandBuffer : Payload->CommandBuffers)
+			if (CurrentStencilLayout != VK_IMAGE_LAYOUT_UNDEFINED)
 			{
-				CommandBufferHandles.Add(SecondaryCommandBuffer->GetHandle());
+				StencilAttachmentInfo.imageLayout = CurrentStencilLayout;
+				StencilAttachmentInfo.loadOp = RenderTargetLoadActionToVulkan(GetLoadAction(GetStencilActions(InRenderPassInfo.DepthStencilRenderTarget.Action)));
+				StencilAttachmentInfo.storeOp = RenderTargetStoreActionToVulkan(GetStoreAction(GetStencilActions(InRenderPassInfo.DepthStencilRenderTarget.Action)));
 
-#if RHI_NEW_GPU_PROFILER
-				// :todo: Flush the event stream if there is one (only use the parent timings for now)
-				UE::RHI::GPUProfiler::FEventStream Dummy(Device.GetQueue(EVulkanQueueType::Graphics)->GetProfilerQueue());
-				SecondaryCommandBuffer->FlushProfilerEvents(Dummy, FPlatformTime::Cycles64());
-#endif
+				bIsStencilStore = (StencilAttachmentInfo.storeOp == VK_ATTACHMENT_STORE_OP_STORE);
+			}
+			else
+			{
+				// Unused image aspects with undefined layout should just remain untouched
+				StencilAttachmentInfo.imageLayout = VK_IMAGE_LAYOUT_ATTACHMENT_OPTIMAL;
+				StencilAttachmentInfo.loadOp = VK_ATTACHMENT_LOAD_OP_NONE;
+				StencilAttachmentInfo.storeOp = VK_ATTACHMENT_STORE_OP_NONE;
 			}
-			ParentCommandBuffer.ExecutedSecondaryCommandBuffers.Append(MoveTemp(Payload->CommandBuffers));
 
-			for (const FVulkanSyncPointRef& Sync : Payload->SyncPoints)
+			if (!DepthView)
 			{
-				AddPendingSyncPoint(Sync);
+				DepthView = FVulkanFramebuffer::GetDepthStencilTargetViewDesc(Texture, NumColorRenderTargets, MipIndex, NumLayers);
 			}
+			StencilAttachmentInfo.imageView = DepthView->GetTextureView().View;
+		}
 
-			delete Payload;
+		if (FVulkanPlatform::RequiresDepthStencilFullWrite() &&
+			VKHasAllFlags(Texture->GetFullAspectMask(), (VK_IMAGE_ASPECT_DEPTH_BIT | VK_IMAGE_ASPECT_STENCIL_BIT)) &&
+			(bIsDepthStore || bIsStencilStore))
+		{
+			// Workaround for old mali drivers: writing not all of the image aspects to compressed render-target could cause gpu-hang
+			DepthAttachmentInfo.storeOp = VK_ATTACHMENT_STORE_OP_STORE;
+			StencilAttachmentInfo.storeOp = VK_ATTACHMENT_STORE_OP_STORE;
 		}
 
-		if (CommandBufferHandles.Num())
+		// CustomResolveSubpass can have targets with a different NumSamples
+		ensure(!NumSamples || (NumSamples == InRenderPassInfo.DepthStencilRenderTarget.DepthStencilTarget->GetNumSamples()) || bHasCustomResolveSubpass);
+		NumSamples = InRenderPassInfo.DepthStencilRenderTarget.DepthStencilTarget->GetNumSamples();
+		const VkSampleCountFlagBits Samples = static_cast<VkSampleCountFlagBits>(NumSamples);
+
+		if (GRHISupportsDepthStencilResolve && (Samples > VK_SAMPLE_COUNT_1_BIT) && InRenderPassInfo.DepthStencilRenderTarget.ResolveTarget)
+		{
+			FVulkanTexture* ResolveTexture = ResourceCast(InRenderPassInfo.DepthStencilRenderTarget.ResolveTarget);
+
+			if (FVulkanView* ResolveView = FVulkanFramebuffer::GetDepthStencilResolveTargetViewDesc(ResolveTexture, MipIndex))
+			{
+				DepthAttachmentInfo.resolveMode = VK_RESOLVE_MODE_SAMPLE_ZERO_BIT;
+				DepthAttachmentInfo.resolveImageLayout = CurrentDepthLayout;
+				DepthAttachmentInfo.imageView = ResolveView->GetTextureView().View;
+				StencilAttachmentInfo.resolveMode = VK_RESOLVE_MODE_SAMPLE_ZERO_BIT;
+				StencilAttachmentInfo.resolveImageLayout = CurrentStencilLayout;
+				StencilAttachmentInfo.imageView = ResolveView->GetTextureView().View;
+			}
+		}
+	}
+
+	if ((NumColorRenderTargets == 0) && !RenderingInfo.pDepthAttachment && !RenderingInfo.pStencilAttachment)
+	{
+		// No Depth and no color, it's a raster-only pass so make sure the renderArea will be set up properly
+		VkExtent2D& Extent = RenderingInfo.renderArea.extent;
+		Extent.width = InRenderPassInfo.ResolveRect.X2 - InRenderPassInfo.ResolveRect.X1;
+		Extent.height = InRenderPassInfo.ResolveRect.Y2 - InRenderPassInfo.ResolveRect.Y1;
+		NumLayers = 1;
+
+		VkOffset2D& Offset = RenderingInfo.renderArea.offset;
+		Offset.x = InRenderPassInfo.ResolveRect.X1;
+		Offset.y = InRenderPassInfo.ResolveRect.Y1;
+	}
+
+	if (GRHISupportsAttachmentVariableRateShading && InRenderPassInfo.ShadingRateTexture)
+	{
+		ZeroVulkanStruct(FragmentShadingRateAttachmentInfo, VK_STRUCTURE_TYPE_RENDERING_FRAGMENT_DENSITY_MAP_ATTACHMENT_INFO_EXT);
+		RenderingInfo.pNext = &FragmentShadingRateAttachmentInfo;
+
+		FVulkanTexture* Texture = ResourceCast(InRenderPassInfo.ShadingRateTexture);
+		check(Texture->GetFormat() == GRHIVariableRateShadingImageFormat);
+		FVulkanView* View = FVulkanFramebuffer::GetFragmentDensityAttachmentViewDesc(Texture, MipIndex);
+		FragmentShadingRateAttachmentInfo.imageView = View->GetTextureView().View;
+
+		if (ValidateShadingRateDataType())
 		{
+			if (GRHIVariableRateShadingImageDataType == VRSImage_Palette)
+			{
+				FragmentShadingRateAttachmentInfo.imageLayout = VK_IMAGE_LAYOUT_FRAGMENT_SHADING_RATE_ATTACHMENT_OPTIMAL_KHR;
+			}
+			if (GRHIVariableRateShadingImageDataType == VRSImage_Fractional)
+			{
+				FragmentShadingRateAttachmentInfo.imageLayout = VK_IMAGE_LAYOUT_FRAGMENT_DENSITY_MAP_OPTIMAL_EXT;
+			}
+		}
+	}
+
+	RenderingInfo.viewMask = (0b1 << InRenderPassInfo.MultiViewCount) - 1;
+	RenderingInfo.layerCount = NumLayers;
+}
+
+
+void FVulkanCommandListContext::RHIBeginParallelRenderPass(TSharedPtr<FRHIParallelRenderPassInfo> InInfo, const TCHAR* InName)
+{
+	checkf(CurrentParallelRenderPassInfo == nullptr, TEXT("There is already a parallel render pass in progress!"));
+	CurrentParallelRenderPassInfo = new FVulkanParallelRenderPassInfo();
+
+	if (Device.GetOptionalExtensions().HasKHRDynamicRendering)
+	{
+		check(InInfo->NumOcclusionQueries == 0);
+		RenderPassInfo = *InInfo;
+
+		CurrentParallelRenderPassInfo->DynamicRenderingInfo = new FVulkanDynamicRenderingInfo();
+		FillDynamicRenderingInfo(*InInfo.Get(), *CurrentParallelRenderPassInfo->DynamicRenderingInfo);
+
+		// Set the flags for the parallel contexts to use directly
+		CurrentParallelRenderPassInfo->DynamicRenderingInfo->RenderingInfo.flags = VK_RENDERING_RESUMING_BIT | VK_RENDERING_SUSPENDING_BIT;
+	}
+	else
+	{
+		RHIBeginRenderPass(*InInfo.Get(), InName);
+		CurrentParallelRenderPassInfo->RenderPassHandle = CurrentRenderPass->GetHandle();
+	}
+
+	InInfo->RHIPlatformData = CurrentParallelRenderPassInfo;
+}
+
+void FVulkanCommandListContext::RHIEndParallelRenderPass()
+{
+	if (!CurrentParallelRenderPassInfo->SecondaryPayloads.IsEmpty())
+	{
+		// Merge all the payloads into a single one
+		FVulkanPayload* MergedPayload = CurrentParallelRenderPassInfo->SecondaryPayloads.Pop();
+		for (FVulkanPayload* OtherPayload : CurrentParallelRenderPassInfo->SecondaryPayloads)
+		{
+			MergedPayload->Merge(*OtherPayload);
+			delete OtherPayload;
+		}
+		CurrentParallelRenderPassInfo->SecondaryPayloads.Reset();
+
+		if (Device.GetOptionalExtensions().HasKHRDynamicRendering)
+		{
+			// Sort the command buffer so that we have correct placement for the first and last based on suspend/resume flags
+			MergedPayload->CommandBuffers.Sort([](const FVulkanCommandBuffer& A, const FVulkanCommandBuffer& B)
+				{
+					// Make sure our first command buffer has no RESUME
+					if (!VKHasAllFlags(A.LastDynamicRenderingFlags, VK_RENDERING_RESUMING_BIT))
+					{
+						return true;
+					}
+					else if (!VKHasAllFlags(B.LastDynamicRenderingFlags, VK_RENDERING_RESUMING_BIT))
+					{
+						return false;
+					}
+
+					// Make sure our last command buffer has no SUSPEND
+					if (!VKHasAllFlags(A.LastDynamicRenderingFlags, VK_RENDERING_SUSPENDING_BIT))
+					{
+						return false;
+					}
+					else if (!VKHasAllFlags(B.LastDynamicRenderingFlags, VK_RENDERING_SUSPENDING_BIT))
+					{
+						return true;
+					}
+
+					// the others don't matter, sort by handle
+					return A.GetHandle() < B.GetHandle();
+				});
+
+			// If the parallel pass has a single payload, it's possible we leave things in a suspended state
+			const bool bNeedsResume = VKHasAllFlags(MergedPayload->CommandBuffers.Last()->LastDynamicRenderingFlags, VK_RENDERING_SUSPENDING_BIT);
+
+			// Insert all the parallel payloads
+			AppendParallelRenderPayload(MergedPayload);
+
+			// Close the last suspended render pass
+			if (bNeedsResume)
+			{
+				CurrentParallelRenderPassInfo->DynamicRenderingInfo->RenderingInfo.flags = VK_RENDERING_RESUMING_BIT;
+				GetCommandBuffer().BeginDynamicRendering(CurrentParallelRenderPassInfo->DynamicRenderingInfo->RenderingInfo);
+				GetCommandBuffer().EndDynamicRendering();
+			}
+		}
+		else
+		{
+			FVulkanPayload& ParentPayload = GetPayload(EPhase::Execute);
+			FVulkanCommandBuffer& ParentCommandBuffer = GetCommandBuffer();
+
+			TArray<VkCommandBuffer> CommandBufferHandles;
+			CommandBufferHandles.Reserve(MergedPayload->CommandBuffers.Num());
+			for (FVulkanCommandBuffer* SecondaryCommandBuffer : MergedPayload->CommandBuffers)
+			{
+				CommandBufferHandles.Add(SecondaryCommandBuffer->GetHandle());
+			}
+			ParentCommandBuffer.ExecutedSecondaryCommandBuffers.Append(MoveTemp(MergedPayload->CommandBuffers));
+
 			VulkanRHI::vkCmdExecuteCommands(ParentCommandBuffer.GetHandle(), CommandBufferHandles.Num(), CommandBufferHandles.GetData());
+
+			ParentPayload.Merge(*MergedPayload);
+			delete MergedPayload;
 		}
 	}
 
-	RHIEndRenderPass();
+	if (!Device.GetOptionalExtensions().HasKHRDynamicRendering)
+	{
+		RHIEndRenderPass();
+	}
 
 	delete CurrentParallelRenderPassInfo;
 	CurrentParallelRenderPassInfo = nullptr;
diff --git a/Engine/Source/Runtime/VulkanRHI/Private/VulkanContext.h b/Engine/Source/Runtime/VulkanRHI/Private/VulkanContext.h
--- a/Engine/Source/Runtime/VulkanRHI/Private/VulkanContext.h
+++ b/Engine/Source/Runtime/VulkanRHI/Private/VulkanContext.h
@@ -32,10 +32,21 @@ enum class EVulkanFlushFlags
 ENUM_CLASS_FLAGS(EVulkanFlushFlags)
 
 
+struct FVulkanDynamicRenderingInfo
+{
+	TStaticArray<VkRenderingAttachmentInfo, MaxSimultaneousRenderTargets> ColorAttachmentInfos;
+	VkRenderingAttachmentInfo DepthAttachmentInfo;
+	VkRenderingAttachmentInfo StencilAttachmentInfo;
+	VkRenderingFragmentShadingRateAttachmentInfoKHR FragmentShadingRateAttachmentInfo;
+	VkRenderingInfo RenderingInfo;
+};
+
 struct FVulkanParallelRenderPassInfo
 {
 	VkRenderPass RenderPassHandle = VK_NULL_HANDLE;
+	FVulkanDynamicRenderingInfo* DynamicRenderingInfo = nullptr;
 	TArray<FVulkanPayload*> SecondaryPayloads;
+	std::atomic<int32> NumParallelContexts = 0;
 };
 
 
@@ -196,6 +207,11 @@ public:
 		return nullptr;
 	}
 
+protected:
+
+	// End the current payload, insert the parallel rendering work, and leave context with a valid payload
+	void AppendParallelRenderPayload(FVulkanPayload* ParallelRenderingPayload);
+
 private:
 
 	FVulkanQueryPool* GetCurrentTimestampQueryPool(FVulkanPayload& Payload);
@@ -401,7 +417,12 @@ public:
 
 	bool AcquirePoolSetAndDescriptorsIfNeeded(const class FVulkanDescriptorSetsLayout& Layout, bool bNeedDescriptors, VkDescriptorSet* OutDescriptors);
 
+
 protected:
+	void ExtractDepthStencilLayouts(const FRHIRenderPassInfo& InInfo, VkImageLayout& OutDepthLayout, VkImageLayout& OutStencilLayout);
+	void FillDynamicRenderingInfo(const FRHIRenderPassInfo& RenderPassInfo, FVulkanDynamicRenderingInfo& OutRenderingInfo);
+	static FVulkanPayload* MergePayloads(TArray<FVulkanPayload*>&& Payloads);
+
 	FVulkanCommandListContext* Immediate;
 	const ERHIPipeline RHIPipeline;
 	const bool bSupportsBreadcrumbs;
diff --git a/Engine/Source/Runtime/VulkanRHI/Private/VulkanPipeline.cpp b/Engine/Source/Runtime/VulkanRHI/Private/VulkanPipeline.cpp
--- a/Engine/Source/Runtime/VulkanRHI/Private/VulkanPipeline.cpp
+++ b/Engine/Source/Runtime/VulkanRHI/Private/VulkanPipeline.cpp
@@ -1325,7 +1325,43 @@ bool FVulkanPipelineStateCacheManager::CreateGfxPipelineFromEntry(FVulkanRHIGrap
 	PipelineInfo.pMultisampleState = &MSInfo;
 	PipelineInfo.pViewportState = &VPInfo;
 
-	PipelineInfo.renderPass = PSO->RenderPass->GetHandle();
+	VkPipelineRenderingCreateInfo PipelineRenderingCreateInfo;
+	TStaticArray<VkFormat, MaxSimultaneousRenderTargets> ColorAttachmentFormats;
+	if (Device->GetOptionalExtensions().HasKHRDynamicRendering)
+	{
+		ZeroVulkanStruct(PipelineRenderingCreateInfo, VK_STRUCTURE_TYPE_PIPELINE_RENDERING_CREATE_INFO);
+		PipelineRenderingCreateInfo.pNext = PipelineInfo.pNext;
+		PipelineInfo.pNext = &PipelineRenderingCreateInfo;
+
+		PipelineRenderingCreateInfo.pColorAttachmentFormats = ColorAttachmentFormats.GetData();
+		PipelineRenderingCreateInfo.colorAttachmentCount = PSO->Desc.RenderTargets.NumColorAttachments;
+		for (uint8 ColorIndex = 0; ColorIndex < PSO->Desc.RenderTargets.NumColorAttachments; ++ColorIndex)
+		{
+			const uint32 Attachment = PSO->Desc.RenderTargets.ColorAttachments[ColorIndex].Attachment;
+			const FGfxPipelineDesc::FRenderTargets::FAttachmentDesc& Descripton = PSO->Desc.RenderTargets.Descriptions[Attachment];
+			ColorAttachmentFormats[ColorIndex] = (VkFormat)Descripton.Format;
+		}
+
+		if (PSO->Desc.RenderTargets.bHasDepthStencil)
+		{
+			// Depth
+			const uint32 Attachment = PSO->Desc.RenderTargets.Depth.Attachment;
+			const FGfxPipelineDesc::FRenderTargets::FAttachmentDesc& Descripton = PSO->Desc.RenderTargets.Descriptions[Attachment];
+			PipelineRenderingCreateInfo.depthAttachmentFormat = (VkFormat)Descripton.Format;
+
+			// Stencil
+			if (VulkanRHI::VulkanFormatHasStencil(PipelineRenderingCreateInfo.depthAttachmentFormat))
+			{
+				// :todo-jn: leave unused attachements to VK_FORMAT_UNDEFINED
+				PipelineRenderingCreateInfo.stencilAttachmentFormat = PipelineRenderingCreateInfo.depthAttachmentFormat;
+			}
+		}
+	}
+	else
+	{
+		PipelineInfo.renderPass = PSO->RenderPass->GetHandle();
+	}
+
 	PipelineInfo.subpass = GfxEntry->SubpassIndex;
 
 	VkPipelineInputAssemblyStateCreateInfo InputAssembly;
@@ -1357,7 +1393,8 @@ bool FVulkanPipelineStateCacheManager::CreateGfxPipelineFromEntry(FVulkanRHIGrap
 
 	PipelineInfo.pDynamicState = &DynamicState;
 
-	const bool bUsingVariableRateShading = PSO->Desc.ShadingRate != EVRSShadingRate::VRSSR_1x1 || (PSO->RenderPass->Layout.bHasFragmentDensityAttachment && PSO->Desc.Combiner != EVRSRateCombiner::VRSRB_Passthrough);
+	const bool bUsingVariableRateShading = (PSO->Desc.ShadingRate != EVRSShadingRate::VRSSR_1x1) || 
+		(PSO->Desc.RenderTargets.bHasFragmentDensityAttachment && PSO->Desc.Combiner != EVRSRateCombiner::VRSRB_Passthrough);
 
 	VkPipelineFragmentShadingRateStateCreateInfoKHR PipelineFragmentShadingRate;
 	if (GRHISupportsPipelineVariableRateShading && GRHIVariableRateShadingImageDataType == VRSImage_Palette && bUsingVariableRateShading)
@@ -1370,6 +1407,7 @@ bool FVulkanPipelineStateCacheManager::CreateGfxPipelineFromEntry(FVulkanRHIGrap
 		PipelineFragmentShadingRate.combinerOps[0] = PipelineToPrimitiveCombinerOperation;
 		PipelineFragmentShadingRate.combinerOps[1] = VK_FRAGMENT_SHADING_RATE_COMBINER_OP_MAX_KHR;		// @todo: This needs to be specified too.
 
+		PipelineFragmentShadingRate.pNext = PipelineInfo.pNext;
 		PipelineInfo.pNext = (void*)&PipelineFragmentShadingRate;
 	}
 
@@ -2131,7 +2169,12 @@ FGraphicsPipelineStateRHIRef FVulkanPipelineStateCacheManager::RHICreateGraphics
 			NewPSO->Layout = Layout;
 			NewPSO->bHasInputAttachments = Layout->GetDescriptorSetsLayout().HasInputAttachments();
 		}
-		NewPSO->RenderPass = Device->GetImmediateContext().PrepareRenderPassForPSOCreation(Initializer);
+
+		if (!Device->GetOptionalExtensions().HasKHRDynamicRendering)
+		{
+			NewPSO->RenderPass = Device->GetImmediateContext().PrepareRenderPassForPSOCreation(Initializer);
+		}
+		
 		{
 			const FBoundShaderStateInput& BSI = Initializer.BoundShaderState;
 			for (int32 StageIdx = 0; StageIdx < ShaderStage::NumGraphicsStages; ++StageIdx)
diff --git a/Engine/Source/Runtime/VulkanRHI/Private/VulkanQueue.cpp b/Engine/Source/Runtime/VulkanRHI/Private/VulkanQueue.cpp
--- a/Engine/Source/Runtime/VulkanRHI/Private/VulkanQueue.cpp
+++ b/Engine/Source/Runtime/VulkanRHI/Private/VulkanQueue.cpp
@@ -273,6 +273,10 @@ void FVulkanQueue::SubmitPayloads(TArrayView<FVulkanPayload*> Payloads, TMap<VkS
 				CommandBufferStorage.Add(CommandBuffer->GetHandle());
 #if RHI_NEW_GPU_PROFILER
 				CommandBuffer->FlushProfilerEvents(Payload->EventStream, SubmitTime);
+				for (FVulkanCommandBuffer* SecondaryCommandBuffer : CommandBuffer->ExecutedSecondaryCommandBuffers)
+				{
+					SecondaryCommandBuffer->FlushProfilerEvents(Payload->EventStream, SubmitTime);
+				}
 #endif // RHI_NEW_GPU_PROFILER
 				CommandBuffer->SetSubmitted();
 			}
diff --git a/Engine/Source/Runtime/VulkanRHI/Private/VulkanRenderTarget.cpp b/Engine/Source/Runtime/VulkanRHI/Private/VulkanRenderTarget.cpp
--- a/Engine/Source/Runtime/VulkanRHI/Private/VulkanRenderTarget.cpp
+++ b/Engine/Source/Runtime/VulkanRHI/Private/VulkanRenderTarget.cpp
@@ -509,25 +509,13 @@ bool FVulkanCommandListContext::IsSwapchainImage(FRHITexture* InTexture) const
 	return false;
 }
 
-void FVulkanCommandListContext::RHIBeginRenderPass(const FRHIRenderPassInfo& InInfo, const TCHAR* InName)
+void FVulkanCommandListContext::ExtractDepthStencilLayouts(const FRHIRenderPassInfo& InInfo, VkImageLayout& OutDepthLayout, VkImageLayout& OutStencilLayout)
 {
-	RenderPassInfo = InInfo;
-
-	if (InInfo.NumOcclusionQueries > 0)
-	{
-		BeginOcclusionQueryBatch(InInfo.NumOcclusionQueries);
-	}
-
-	// Acquire a descriptor pool set on a first render pass
-	AcquirePoolSetContainer();
-
-	const bool bNeedsAllPlanes = Device.NeedsAllPlanes();
-
 	FRHITexture* DSTexture = InInfo.DepthStencilRenderTarget.DepthStencilTarget;
-	VkImageLayout CurrentDepthLayout = VK_IMAGE_LAYOUT_UNDEFINED;
-	VkImageLayout CurrentStencilLayout = VK_IMAGE_LAYOUT_UNDEFINED;
 	if (DSTexture)
 	{
+		const bool bNeedsAllPlanes = Device.NeedsAllPlanes();
+
 		FVulkanTexture& VulkanTexture = *ResourceCast(DSTexture);
 		const VkImageAspectFlags AspectFlags = VulkanTexture.GetFullAspectMask();
 
@@ -536,15 +524,15 @@ void FVulkanCommandListContext::RHIBeginRenderPass(const FRHIRenderPassInfo& InI
 		{
 			if (ExclusiveDepthStencil.IsDepthWrite())
 			{
-				CurrentDepthLayout = VK_IMAGE_LAYOUT_ATTACHMENT_OPTIMAL;
+				OutDepthLayout = VK_IMAGE_LAYOUT_ATTACHMENT_OPTIMAL;
 			}
 			else if (ExclusiveDepthStencil.IsDepthRead())
 			{
-				CurrentDepthLayout = VK_IMAGE_LAYOUT_READ_ONLY_OPTIMAL;
+				OutDepthLayout = VK_IMAGE_LAYOUT_READ_ONLY_OPTIMAL;
 			}
 			else if (bNeedsAllPlanes)
 			{
-				CurrentDepthLayout = FVulkanPipelineBarrier::GetDepthOrStencilLayout(VulkanTexture.AllPlanesTrackedAccess[0]);
+				OutDepthLayout = FVulkanPipelineBarrier::GetDepthOrStencilLayout(VulkanTexture.AllPlanesTrackedAccess[0]);
 			}
 		}
 
@@ -552,59 +540,111 @@ void FVulkanCommandListContext::RHIBeginRenderPass(const FRHIRenderPassInfo& InI
 		{
 			if (ExclusiveDepthStencil.IsStencilWrite())
 			{
-				CurrentStencilLayout = VK_IMAGE_LAYOUT_ATTACHMENT_OPTIMAL;
+				OutStencilLayout = VK_IMAGE_LAYOUT_ATTACHMENT_OPTIMAL;
 			}
 			else if (ExclusiveDepthStencil.IsStencilRead())
 			{
-				CurrentStencilLayout = VK_IMAGE_LAYOUT_READ_ONLY_OPTIMAL;
+				OutStencilLayout = VK_IMAGE_LAYOUT_READ_ONLY_OPTIMAL;
 			}
 			else if (bNeedsAllPlanes)
 			{
-				CurrentStencilLayout = FVulkanPipelineBarrier::GetDepthOrStencilLayout(VulkanTexture.AllPlanesTrackedAccess[1]);
+				OutStencilLayout = FVulkanPipelineBarrier::GetDepthOrStencilLayout(VulkanTexture.AllPlanesTrackedAccess[1]);
 			}
 		}
 	}
+}
 
-	FVulkanRenderTargetLayout RTLayout(Device, InInfo, CurrentDepthLayout, CurrentStencilLayout);
-	check(RTLayout.GetExtent2D().width != 0 && RTLayout.GetExtent2D().height != 0);
+void FVulkanCommandListContext::RHIBeginRenderPass(const FRHIRenderPassInfo& InInfo, const TCHAR* InName)
+{
+	RenderPassInfo = InInfo;
 
-	FVulkanRenderPass* RenderPass = Device.GetRenderPassManager().GetOrCreateRenderPass(RTLayout);
-	FRHISetRenderTargetsInfo RTInfo;
-	InInfo.ConvertToRenderTargetsInfo(RTInfo);
+	if (InInfo.NumOcclusionQueries > 0)
+	{
+		BeginOcclusionQueryBatch(InInfo.NumOcclusionQueries);
+	}
 
-	FVulkanFramebuffer* Framebuffer = Device.GetRenderPassManager().GetOrCreateFramebuffer(RTInfo, RTLayout, *RenderPass);
-	checkf(RenderPass != nullptr && Framebuffer != nullptr, TEXT("RenderPass not started! Bad combination of values? Depth %p #Color %d Color0 %p"), (void*)InInfo.DepthStencilRenderTarget.DepthStencilTarget, InInfo.GetNumColorRenderTargets(), (void*)InInfo.ColorRenderTargets[0].RenderTarget);
+	// Acquire a descriptor pool set on a first render pass
+	AcquirePoolSetContainer();
 
 	const bool bIsParallelRenderPass = (CurrentParallelRenderPassInfo != nullptr);
-	FVulkanBeginRenderPassInfo BeginRenderPassInfo{*RenderPass, *Framebuffer, bIsParallelRenderPass};
-	Device.GetRenderPassManager().BeginRenderPass(*this, InInfo, RTLayout, BeginRenderPassInfo);
+	if (Device.GetOptionalExtensions().HasKHRDynamicRendering)
+	{
+		check(!CurrentRenderPass && !CurrentFramebuffer);
 
-	check(!CurrentRenderPass);
-	CurrentRenderPass = RenderPass;
-	CurrentFramebuffer = Framebuffer;
+		VkRect2D RenderArea;
+		if (bIsParallelRenderPass)
+		{
+			check(CurrentParallelRenderPassInfo->DynamicRenderingInfo);
+			GetCommandBuffer().BeginDynamicRendering(CurrentParallelRenderPassInfo->DynamicRenderingInfo->RenderingInfo);
+			RenderArea = CurrentParallelRenderPassInfo->DynamicRenderingInfo->RenderingInfo.renderArea;
+		}
+		else
+		{
+			FVulkanDynamicRenderingInfo DynamicRenderingInfo;
+			FillDynamicRenderingInfo(InInfo, DynamicRenderingInfo);
+			GetCommandBuffer().BeginDynamicRendering(DynamicRenderingInfo.RenderingInfo);
+			RenderArea = DynamicRenderingInfo.RenderingInfo.renderArea;
+		}
+
+		GetPendingGfxState()->SetViewport(RenderArea.offset.x, RenderArea.offset.y, 0, RenderArea.extent.width, RenderArea.extent.height, 1);
+	}
+	else
+	{
+		VkImageLayout CurrentDepthLayout = VK_IMAGE_LAYOUT_UNDEFINED;
+		VkImageLayout CurrentStencilLayout = VK_IMAGE_LAYOUT_UNDEFINED;
+		ExtractDepthStencilLayouts(InInfo, CurrentDepthLayout, CurrentStencilLayout);
+
+		FVulkanRenderTargetLayout RTLayout(Device, InInfo, CurrentDepthLayout, CurrentStencilLayout);
+		check(RTLayout.GetExtent2D().width != 0 && RTLayout.GetExtent2D().height != 0);
+
+		FVulkanRenderPass* RenderPass = Device.GetRenderPassManager().GetOrCreateRenderPass(RTLayout);
+		FRHISetRenderTargetsInfo RTInfo;
+		InInfo.ConvertToRenderTargetsInfo(RTInfo);
+
+		FVulkanFramebuffer* Framebuffer = Device.GetRenderPassManager().GetOrCreateFramebuffer(RTInfo, RTLayout, *RenderPass);
+		checkf(RenderPass != nullptr && Framebuffer != nullptr, TEXT("RenderPass not started! Bad combination of values? Depth %p #Color %d Color0 %p"), (void*)InInfo.DepthStencilRenderTarget.DepthStencilTarget, InInfo.GetNumColorRenderTargets(), (void*)InInfo.ColorRenderTargets[0].RenderTarget);
+
+		FVulkanBeginRenderPassInfo BeginRenderPassInfo{ *RenderPass, *Framebuffer, bIsParallelRenderPass };
+		Device.GetRenderPassManager().BeginRenderPass(*this, InInfo, RTLayout, BeginRenderPassInfo);
+
+		check(!CurrentRenderPass);
+		CurrentRenderPass = RenderPass;
+		CurrentFramebuffer = Framebuffer;
+	}
 }
 
 void FVulkanCommandListContext::RHIEndRenderPass()
 {
-	Device.GetRenderPassManager().EndRenderPass(*this);
-
-	const bool bHasOcclusionQueries = (RenderPassInfo.NumOcclusionQueries > 0);
-	if (bHasOcclusionQueries)
+	if (Device.GetOptionalExtensions().HasKHRDynamicRendering)
 	{
-		// Force the syncs points to be signaled right after the render pass containing the queries
-		FlushPendingSyncPoints();
+		GetCommandBuffer().EndDynamicRendering();
 	}
+	else
+	{
+		Device.GetRenderPassManager().EndRenderPass(*this);
 
-	check(CurrentRenderPass);
-	CurrentRenderPass = nullptr;
+		check(CurrentRenderPass);
+		CurrentRenderPass = nullptr;
+	}
 }
 
 void FVulkanCommandListContext::RHINextSubpass()
 {
-	check(CurrentRenderPass);
-	FVulkanCommandBuffer& CommandBuffer = GetCommandBuffer();
-	VkCommandBuffer CommandBufferHandle = CommandBuffer.GetHandle();
-	VulkanRHI::vkCmdNextSubpass(CommandBufferHandle, VK_SUBPASS_CONTENTS_INLINE);
+	if (Device.GetOptionalExtensions().HasKHRDynamicRendering)
+	{
+		// :todo: Need to properly handle local_read (VK_IMAGE_LAYOUT_RENDERING_LOCAL_READ)
+		VkMemoryBarrier Barrier = { VK_STRUCTURE_TYPE_MEMORY_BARRIER, nullptr, VK_ACCESS_SHADER_WRITE_BIT, VK_ACCESS_SHADER_READ_BIT };
+		VulkanRHI::vkCmdPipelineBarrier(GetCommandBuffer().GetHandle(), 
+			VK_PIPELINE_STAGE_FRAGMENT_SHADER_BIT, VK_PIPELINE_STAGE_FRAGMENT_SHADER_BIT, 
+			VK_DEPENDENCY_BY_REGION_BIT, 1, &Barrier, 0, nullptr, 0, nullptr);
+	}
+	else
+	{
+		check(CurrentRenderPass);
+		FVulkanCommandBuffer& CommandBuffer = GetCommandBuffer();
+		VkCommandBuffer CommandBufferHandle = CommandBuffer.GetHandle();
+		VulkanRHI::vkCmdNextSubpass(CommandBufferHandle, VK_SUBPASS_CONTENTS_INLINE);
+	}
 }
 
 // Need a separate struct so we can memzero/remove dependencies on reference counts
diff --git a/Engine/Source/Runtime/VulkanRHI/Private/VulkanSubmission.cpp b/Engine/Source/Runtime/VulkanRHI/Private/VulkanSubmission.cpp
--- a/Engine/Source/Runtime/VulkanRHI/Private/VulkanSubmission.cpp
+++ b/Engine/Source/Runtime/VulkanRHI/Private/VulkanSubmission.cpp
@@ -30,6 +30,25 @@ FVulkanPayload::~FVulkanPayload()
 	}
 }
 
+void FVulkanPayload::Merge(FVulkanPayload& OtherPayload)
+{
+	check(OtherPayload.SignalSemaphores.Num() == 0);
+	check(OtherPayload.WaitSemaphores.Num() == 0);
+
+	for (int32 QueryPoolTypeIndex = 0; QueryPoolTypeIndex < OtherPayload.QueryPools.Num(); ++QueryPoolTypeIndex)
+	{
+		QueryPools[QueryPoolTypeIndex].Append(MoveTemp(OtherPayload.QueryPools[QueryPoolTypeIndex]));
+	}
+
+	DescriptorPoolSetContainers.Append(MoveTemp(OtherPayload.DescriptorPoolSetContainers));
+
+	checkf(OtherPayload.CommandBuffers.IsEmpty() || (OtherPayload.CommandBuffers[0]->GetCommandBufferType() != EVulkanCommandBufferType::Primary),
+		TEXT("Merging payloads should only be used on Secondary or Parallel command buffers."));
+	CommandBuffers.Append(MoveTemp(OtherPayload.CommandBuffers));
+
+	SyncPoints.Append(MoveTemp(OtherPayload.SyncPoints));
+}
+
 
 
 class FVulkanThread final : private FRunnable
@@ -250,7 +269,28 @@ IRHIComputeContext* FVulkanDynamicRHI::RHIGetParallelCommandContext(FRHIParallel
 {
 	FVulkanParallelRenderPassInfo* VulkanParallelRenderPassInfo = (FVulkanParallelRenderPassInfo*)ParallelRenderPass.RHIPlatformData;
 	checkf(VulkanParallelRenderPassInfo, TEXT("Must begin parallel render pass on parent context before creating parallel contexts!"));
-	return new FVulkanCommandListContext(*Device, &Device->GetImmediateContext(), VulkanParallelRenderPassInfo);
+	FVulkanCommandListContext* ParallelCommandContext = new FVulkanCommandListContext(*Device, &Device->GetImmediateContext(), VulkanParallelRenderPassInfo);
+	ParallelCommandContext->RenderPassInfo = ParallelRenderPass;
+
+	if (Device->GetOptionalExtensions().HasKHRDynamicRendering)
+	{
+		check(VulkanParallelRenderPassInfo->DynamicRenderingInfo);
+
+		const int32 ContextIndex = VulkanParallelRenderPassInfo->NumParallelContexts++;
+		if (ContextIndex < 2)
+		{
+			// Make sure to create bookends (first context who only suspends, last context who only resumes)
+			VkRenderingInfo RenderingInfo = VulkanParallelRenderPassInfo->DynamicRenderingInfo->RenderingInfo;
+			RenderingInfo.flags = (ContextIndex == 0) ? VK_RENDERING_SUSPENDING_BIT : VK_RENDERING_RESUMING_BIT;
+			ParallelCommandContext->GetCommandBuffer().BeginDynamicRendering(RenderingInfo);
+		}
+		else
+		{
+			ParallelCommandContext->GetCommandBuffer().BeginDynamicRendering(VulkanParallelRenderPassInfo->DynamicRenderingInfo->RenderingInfo);
+		}
+	}
+
+	return ParallelCommandContext;
 }
 
 static FCriticalSection GSecondaryPayloadsCS;
@@ -258,6 +298,11 @@ IRHIPlatformCommandList* FVulkanDynamicRHI::RHIFinalizeParallelContext(IRHICompu
 {
 	FVulkanCommandListContext* VulkanContext = static_cast<FVulkanCommandListContext*>(Context);
 
+	if (Device->GetOptionalExtensions().HasKHRDynamicRendering)
+	{
+		VulkanContext->GetCommandBuffer().EndDynamicRendering();
+	}
+
 	// Call Finalize under this lock
 	{
 		FScopeLock Lock(&GSecondaryPayloadsCS);
diff --git a/Engine/Source/Runtime/VulkanRHI/Private/VulkanSubmission.h b/Engine/Source/Runtime/VulkanRHI/Private/VulkanSubmission.h
--- a/Engine/Source/Runtime/VulkanRHI/Private/VulkanSubmission.h
+++ b/Engine/Source/Runtime/VulkanRHI/Private/VulkanSubmission.h
@@ -46,6 +46,9 @@ public:
 
 protected:
 
+	// Must only be used to merge Secondary or Parallel command buffers
+	void Merge(FVulkanPayload& Other);
+
 	void PreExecute();
 
 	FVulkanQueue& Queue;
