VulkanRHI: Load extended_dynamic_state extensions and shader object extension. (WIP)

#rnx

diff --git a/Engine/Source/Runtime/VulkanRHI/Private/VulkanCommandWrappers.h b/Engine/Source/Runtime/VulkanRHI/Private/VulkanCommandWrappers.h
--- a/Engine/Source/Runtime/VulkanRHI/Private/VulkanCommandWrappers.h
+++ b/Engine/Source/Runtime/VulkanRHI/Private/VulkanCommandWrappers.h
@@ -265,6 +265,32 @@ struct FWrapLayer
 	static void GetSemaphoreCounterValueKHR(VkResult Result, VkDevice Device, VkSemaphore Semaphore, uint64_t* Value) VULKAN_LAYER_BODY
 	static void WaitSemaphoresKHR(VkResult Result, VkDevice Device, const VkSemaphoreWaitInfo* WaitInfo, uint64_t Timeout) VULKAN_LAYER_BODY
 	static void SignalSemaphoreKHR(VkResult Result, VkDevice Device, const VkSemaphoreSignalInfo* SignalInfo) VULKAN_LAYER_BODY
+
+	static void CreateShadersEXT(VkResult Result, VkDevice Device, uint32_t CreateInfoCount, const VkShaderCreateInfoEXT* CreateInfos, const VkAllocationCallbacks* Allocator, VkShaderEXT* Shaders) VULKAN_LAYER_BODY
+	static void DestroyShaderEXT(VkResult Result, VkDevice Device, VkShaderEXT Shader, const VkAllocationCallbacks* Allocator) VULKAN_LAYER_BODY
+	static void GetShaderBinaryDataEXT(VkResult Result, VkDevice Device, VkShaderEXT Shader, size_t* DataSize, void* Data) VULKAN_LAYER_BODY
+	static void CmdBindShadersEXT(VkResult Result, VkCommandBuffer CommandBuffer, uint32_t StageCount, const VkShaderStageFlagBits* Stages, const VkShaderEXT* Shaders) VULKAN_LAYER_BODY
+	static void CmdBindVertexBuffers2EXT(VkResult Result, VkCommandBuffer CommandBuffer, uint32_t FirstBinding, uint32_t BindingCount, const VkBuffer* Buffers, const VkDeviceSize* Offsets, const VkDeviceSize* Sizes, const VkDeviceSize* Strides) VULKAN_LAYER_BODY
+	static void CmdSetCullModeEXT(VkResult Result, VkCommandBuffer CommandBuffer, VkCullModeFlags CullMode) VULKAN_LAYER_BODY
+	static void CmdSetColorBlendEnableEXT(VkResult Result, VkCommandBuffer CommandBuffer, uint32_t FirstAttachment, uint32_t AttachmentCount, const VkBool32* ColorBlendEnables) VULKAN_LAYER_BODY
+	static void CmdSetColorBlendEquationEXT(VkResult Result, VkCommandBuffer CommandBuffer, uint32_t FirstAttachment, uint32_t AttachmentCount, const VkColorBlendEquationEXT* ColorBlendEquations) VULKAN_LAYER_BODY
+	static void CmdSetColorWriteMaskEXT(VkResult Result, VkCommandBuffer CommandBuffer, uint32_t FirstAttachment, uint32_t AttachmentCount, const VkColorComponentFlags* ColorWriteMasks) VULKAN_LAYER_BODY
+	static void CmdSetFrontFaceEXT(VkResult Result, VkCommandBuffer CommandBuffer, VkFrontFace FrontFace) VULKAN_LAYER_BODY
+	static void CmdSetPolygonModeEXT(VkResult Result, VkCommandBuffer CommandBuffer, VkPolygonMode PolygonMode) VULKAN_LAYER_BODY
+	static void CmdSetPrimitiveTopologyEXT(VkResult Result, VkCommandBuffer CommandBuffer, VkPrimitiveTopology PrimitiveTopology) VULKAN_LAYER_BODY
+	static void CmdSetRasterizationSamplesEXT(VkResult Result, VkCommandBuffer CommandBuffer, VkSampleCountFlagBits RasterizationSamples) VULKAN_LAYER_BODY
+	static void CmdSetScissorWithCountEXT(VkResult Result, VkCommandBuffer CommandBuffer, uint32_t ScissorCount, const VkRect2D* Scissors) VULKAN_LAYER_BODY
+	static void CmdSetStencilOpEXT(VkResult Result, VkCommandBuffer CommandBuffer, VkStencilFaceFlags FaceMask, VkStencilOp FailOp, VkStencilOp passOp, VkStencilOp depthFailOp, VkCompareOp compareOp) VULKAN_LAYER_BODY
+	static void CmdSetStencilTestEnableEXT(VkResult Result, VkCommandBuffer CommandBuffer, VkBool32 StencilTestEnable) VULKAN_LAYER_BODY
+	static void CmdSetVertexInputEXT(VkResult Result, VkCommandBuffer CommandBuffer, uint32_t VertexBindingDescriptionCount, const VkVertexInputBindingDescription2EXT* pVertexBindingDescriptions, uint32_t vertexAttributeDescriptionCount, const VkVertexInputAttributeDescription2EXT* pVertexAttributeDescriptions) VULKAN_LAYER_BODY
+	static void CmdSetViewportWithCountEXT(VkResult Result, VkCommandBuffer CommandBuffer, uint32_t ViewportCount, const VkViewport* pViewports) VULKAN_LAYER_BODY
+	static void CmdSetDepthBiasEnableEXT(VkResult Result, VkCommandBuffer CommandBuffer, VkBool32 DepthBiasEnable) VULKAN_LAYER_BODY
+	static void CmdSetDepthBoundsTestEnableEXT(VkResult Result, VkCommandBuffer CommandBuffer, VkBool32 DepthBoundsTestEnable) VULKAN_LAYER_BODY
+	static void CmdSetDepthClampEnableEXT(VkResult Result, VkCommandBuffer CommandBuffer, VkBool32 DepthClampEnable) VULKAN_LAYER_BODY
+	static void CmdSetDepthCompareOpEXT(VkResult Result, VkCommandBuffer CommandBuffer, VkCompareOp DepthCompareOp) VULKAN_LAYER_BODY
+	static void CmdSetDepthTestEnableEXT(VkResult Result, VkCommandBuffer CommandBuffer, VkBool32 DepthTestEnable) VULKAN_LAYER_BODY
+	static void CmdSetDepthWriteEnableEXT(VkResult Result, VkCommandBuffer CommandBuffer, VkBool32 DepthWriteEnable) VULKAN_LAYER_BODY
+	static void CmdSetRasterizerDiscardEnableEXT(VkResult Result, VkCommandBuffer CommandBuffer, VkBool32 RasterizerDiscardEnable) VULKAN_LAYER_BODY
 };
 
 #undef VULKAN_LAYER_BODY
@@ -1841,6 +1867,184 @@ namespace VulkanRHI
 		return Result;
 	}
 
+	static FORCEINLINE_DEBUGGABLE VkResult vkCreateShadersEXT(VkDevice Device, uint32_t CreateInfoCount, const VkShaderCreateInfoEXT* CreateInfos, const VkAllocationCallbacks* Allocator, VkShaderEXT* Shaders)
+	{
+		FWrapLayer::CreateShadersEXT(VK_RESULT_MAX_ENUM, Device, CreateInfoCount, CreateInfos, Allocator, Shaders);
+		const VkResult Result = VULKANAPINAMESPACE::vkCreateShadersEXT(Device, CreateInfoCount, CreateInfos, Allocator, Shaders);
+		FWrapLayer::CreateShadersEXT(Result, Device, CreateInfoCount, CreateInfos, Allocator, Shaders);
+		return Result;
+	}
+
+	static FORCEINLINE_DEBUGGABLE void vkDestroyShaderEXT(VkDevice Device, VkShaderEXT Shader, const VkAllocationCallbacks* Allocator)
+	{
+		FWrapLayer::DestroyShaderEXT(VK_RESULT_MAX_ENUM, Device, Shader, Allocator);
+		VULKANAPINAMESPACE::vkDestroyShaderEXT(Device, Shader, Allocator);
+		FWrapLayer::DestroyShaderEXT(VK_SUCCESS, Device, Shader, Allocator);
+	}
+
+	static FORCEINLINE_DEBUGGABLE VkResult vkGetShaderBinaryDataEXT(VkDevice Device, VkShaderEXT Shader, size_t* DataSize, void* Data)
+	{
+		FWrapLayer::GetShaderBinaryDataEXT(VK_RESULT_MAX_ENUM, Device, Shader, DataSize, Data);
+		const VkResult Result = VULKANAPINAMESPACE::vkGetShaderBinaryDataEXT(Device, Shader, DataSize, Data);
+		FWrapLayer::GetShaderBinaryDataEXT(Result, Device, Shader, DataSize, Data);
+		return Result;
+	}
+
+	static FORCEINLINE_DEBUGGABLE void vkCmdBindShadersEXT(VkCommandBuffer CommandBuffer, uint32_t StageCount, const VkShaderStageFlagBits* Stages, const VkShaderEXT* Shaders)
+	{
+		FWrapLayer::CmdBindShadersEXT(VK_RESULT_MAX_ENUM, CommandBuffer, StageCount, Stages, Shaders);
+		VULKANAPINAMESPACE::vkCmdBindShadersEXT(CommandBuffer, StageCount, Stages, Shaders);
+		FWrapLayer::CmdBindShadersEXT(VK_SUCCESS, CommandBuffer, StageCount, Stages, Shaders);
+	}
+
+	static FORCEINLINE_DEBUGGABLE void vkCmdBindVertexBuffers2EXT(VkCommandBuffer CommandBuffer, uint32_t FirstBinding, uint32_t BindingCount, const VkBuffer* Buffers, const VkDeviceSize* Offsets, const VkDeviceSize* Sizes, const VkDeviceSize* Strides)
+	{
+		FWrapLayer::CmdBindVertexBuffers2EXT(VK_RESULT_MAX_ENUM, CommandBuffer, FirstBinding, BindingCount, Buffers, Offsets, Sizes, Strides);
+		VULKANAPINAMESPACE::vkCmdBindVertexBuffers2EXT(CommandBuffer, FirstBinding, BindingCount, Buffers, Offsets, Sizes, Strides);
+		FWrapLayer::CmdBindVertexBuffers2EXT(VK_SUCCESS, CommandBuffer, FirstBinding, BindingCount, Buffers, Offsets, Sizes, Strides);
+	}
+
+	static FORCEINLINE_DEBUGGABLE void vkCmdSetCullModeEXT(VkCommandBuffer CommandBuffer, VkCullModeFlags CullMode)
+	{
+		FWrapLayer::CmdSetCullModeEXT(VK_RESULT_MAX_ENUM, CommandBuffer, CullMode);
+		VULKANAPINAMESPACE::vkCmdSetCullModeEXT(CommandBuffer, CullMode);
+		FWrapLayer::CmdSetCullModeEXT(VK_SUCCESS, CommandBuffer, CullMode);
+	}
+
+	static FORCEINLINE_DEBUGGABLE void vkCmdSetColorBlendEnableEXT(VkCommandBuffer CommandBuffer, uint32_t FirstAttachment, uint32_t AttachmentCount, const VkBool32* ColorBlendEnables)
+	{
+		FWrapLayer::CmdSetColorBlendEnableEXT(VK_RESULT_MAX_ENUM, CommandBuffer, FirstAttachment, AttachmentCount, ColorBlendEnables);
+		VULKANAPINAMESPACE::vkCmdSetColorBlendEnableEXT(CommandBuffer, FirstAttachment, AttachmentCount, ColorBlendEnables);
+		FWrapLayer::CmdSetColorBlendEnableEXT(VK_SUCCESS, CommandBuffer, FirstAttachment, AttachmentCount, ColorBlendEnables);
+	}
+
+	static FORCEINLINE_DEBUGGABLE void vkCmdSetColorBlendEquationEXT(VkCommandBuffer CommandBuffer, uint32_t FirstAttachment, uint32_t AttachmentCount, const VkColorBlendEquationEXT* ColorBlendEquations)
+	{
+		FWrapLayer::CmdSetColorBlendEquationEXT(VK_RESULT_MAX_ENUM, CommandBuffer, FirstAttachment, AttachmentCount, ColorBlendEquations);
+		VULKANAPINAMESPACE::vkCmdSetColorBlendEquationEXT(CommandBuffer, FirstAttachment, AttachmentCount, ColorBlendEquations);
+		FWrapLayer::CmdSetColorBlendEquationEXT(VK_SUCCESS, CommandBuffer, FirstAttachment, AttachmentCount, ColorBlendEquations);
+	}
+
+	static FORCEINLINE_DEBUGGABLE void vkCmdSetColorWriteMaskEXT(VkCommandBuffer CommandBuffer, uint32_t FirstAttachment, uint32_t AttachmentCount, const VkColorComponentFlags* ColorWriteMasks)
+	{
+		FWrapLayer::CmdSetColorWriteMaskEXT(VK_RESULT_MAX_ENUM, CommandBuffer, FirstAttachment, AttachmentCount, ColorWriteMasks);
+		VULKANAPINAMESPACE::vkCmdSetColorWriteMaskEXT(CommandBuffer, FirstAttachment, AttachmentCount, ColorWriteMasks);
+		FWrapLayer::CmdSetColorWriteMaskEXT(VK_SUCCESS, CommandBuffer, FirstAttachment, AttachmentCount, ColorWriteMasks);
+	}
+
+	static FORCEINLINE_DEBUGGABLE void vkCmdSetFrontFaceEXT(VkCommandBuffer CommandBuffer, VkFrontFace FrontFace)
+	{
+		FWrapLayer::CmdSetFrontFaceEXT(VK_RESULT_MAX_ENUM, CommandBuffer, FrontFace);
+		VULKANAPINAMESPACE::vkCmdSetFrontFaceEXT(CommandBuffer, FrontFace);
+		FWrapLayer::CmdSetFrontFaceEXT(VK_SUCCESS, CommandBuffer, FrontFace);
+	}
+
+	static FORCEINLINE_DEBUGGABLE void vkCmdSetPolygonModeEXT(VkCommandBuffer CommandBuffer, VkPolygonMode PolygonMode)
+	{
+		FWrapLayer::CmdSetPolygonModeEXT(VK_RESULT_MAX_ENUM, CommandBuffer, PolygonMode);
+		VULKANAPINAMESPACE::vkCmdSetPolygonModeEXT(CommandBuffer, PolygonMode);
+		FWrapLayer::CmdSetPolygonModeEXT(VK_SUCCESS, CommandBuffer, PolygonMode);
+	}
+
+	static FORCEINLINE_DEBUGGABLE void vkCmdSetPrimitiveTopologyEXT(VkCommandBuffer CommandBuffer, VkPrimitiveTopology PrimitiveTopology)
+	{
+		FWrapLayer::CmdSetPrimitiveTopologyEXT(VK_RESULT_MAX_ENUM, CommandBuffer, PrimitiveTopology);
+		VULKANAPINAMESPACE::vkCmdSetPrimitiveTopologyEXT(CommandBuffer, PrimitiveTopology);
+		FWrapLayer::CmdSetPrimitiveTopologyEXT(VK_SUCCESS, CommandBuffer, PrimitiveTopology);
+	}
+
+	static FORCEINLINE_DEBUGGABLE void vkCmdSetRasterizationSamplesEXT(VkCommandBuffer CommandBuffer, VkSampleCountFlagBits RasterizationSamples)
+	{
+		FWrapLayer::CmdSetRasterizationSamplesEXT(VK_RESULT_MAX_ENUM, CommandBuffer, RasterizationSamples);
+		VULKANAPINAMESPACE::vkCmdSetRasterizationSamplesEXT(CommandBuffer, RasterizationSamples);
+		FWrapLayer::CmdSetRasterizationSamplesEXT(VK_SUCCESS, CommandBuffer, RasterizationSamples);
+	}
+
+	static FORCEINLINE_DEBUGGABLE void vkCmdSetScissorWithCountEXT(VkCommandBuffer CommandBuffer, uint32_t ScissorCount, const VkRect2D* Scissors)
+	{
+		FWrapLayer::CmdSetScissorWithCountEXT(VK_RESULT_MAX_ENUM, CommandBuffer, ScissorCount, Scissors);
+		VULKANAPINAMESPACE::vkCmdSetScissorWithCountEXT(CommandBuffer, ScissorCount, Scissors);
+		FWrapLayer::CmdSetScissorWithCountEXT(VK_SUCCESS, CommandBuffer, ScissorCount, Scissors);
+	}
+
+	static FORCEINLINE_DEBUGGABLE void vkCmdSetStencilOpEXT(VkCommandBuffer CommandBuffer, VkStencilFaceFlags FaceMask, VkStencilOp FailOp, VkStencilOp PassOp, VkStencilOp DepthFailOp, VkCompareOp CompareOp)
+	{
+		FWrapLayer::CmdSetStencilOpEXT(VK_RESULT_MAX_ENUM, CommandBuffer, FaceMask, FailOp, PassOp, DepthFailOp, CompareOp);
+		VULKANAPINAMESPACE::vkCmdSetStencilOpEXT(CommandBuffer, FaceMask, FailOp, PassOp, DepthFailOp, CompareOp);
+		FWrapLayer::CmdSetStencilOpEXT(VK_SUCCESS, CommandBuffer, FaceMask, FailOp, PassOp, DepthFailOp, CompareOp);
+	}
+
+	static FORCEINLINE_DEBUGGABLE void vkCmdSetStencilTestEnableEXT(VkCommandBuffer CommandBuffer, VkBool32 StencilTestEnable)
+	{
+		FWrapLayer::CmdSetStencilTestEnableEXT(VK_RESULT_MAX_ENUM, CommandBuffer, StencilTestEnable);
+		VULKANAPINAMESPACE::vkCmdSetStencilTestEnableEXT(CommandBuffer, StencilTestEnable);
+		FWrapLayer::CmdSetStencilTestEnableEXT(VK_SUCCESS, CommandBuffer, StencilTestEnable);
+	}
+
+	static FORCEINLINE_DEBUGGABLE void vkCmdSetVertexInputEXT(VkCommandBuffer CommandBuffer, uint32_t VertexBindingDescriptionCount, const VkVertexInputBindingDescription2EXT* VertexBindingDescriptions, uint32_t VertexAttributeDescriptionCount, const VkVertexInputAttributeDescription2EXT* VertexAttributeDescriptions)
+	{
+		FWrapLayer::CmdSetVertexInputEXT(VK_RESULT_MAX_ENUM, CommandBuffer, VertexBindingDescriptionCount, VertexBindingDescriptions, VertexAttributeDescriptionCount, VertexAttributeDescriptions);
+		VULKANAPINAMESPACE::vkCmdSetVertexInputEXT(CommandBuffer, VertexBindingDescriptionCount, VertexBindingDescriptions, VertexAttributeDescriptionCount, VertexAttributeDescriptions);
+		FWrapLayer::CmdSetVertexInputEXT(VK_SUCCESS, CommandBuffer, VertexBindingDescriptionCount, VertexBindingDescriptions, VertexAttributeDescriptionCount, VertexAttributeDescriptions);
+	}
+
+	static FORCEINLINE_DEBUGGABLE void vkCmdSetViewportWithCountEXT(VkCommandBuffer CommandBuffer, uint32_t ViewportCount, const VkViewport* Viewports)
+	{
+		FWrapLayer::CmdSetViewportWithCountEXT(VK_RESULT_MAX_ENUM, CommandBuffer, ViewportCount, Viewports);
+		VULKANAPINAMESPACE::vkCmdSetViewportWithCountEXT(CommandBuffer, ViewportCount, Viewports);
+		FWrapLayer::CmdSetViewportWithCountEXT(VK_SUCCESS, CommandBuffer, ViewportCount, Viewports);
+	}
+
+	static FORCEINLINE_DEBUGGABLE void vkCmdSetDepthBiasEnableEXT(VkCommandBuffer CommandBuffer, VkBool32 DepthBiasEnable)
+	{
+		FWrapLayer::CmdSetDepthBiasEnableEXT(VK_RESULT_MAX_ENUM, CommandBuffer, DepthBiasEnable);
+		VULKANAPINAMESPACE::vkCmdSetDepthBiasEnableEXT(CommandBuffer, DepthBiasEnable);
+		FWrapLayer::CmdSetDepthBiasEnableEXT(VK_SUCCESS, CommandBuffer, DepthBiasEnable);
+	}
+
+	static FORCEINLINE_DEBUGGABLE void vkCmdSetDepthBoundsTestEnableEXT(VkCommandBuffer CommandBuffer, VkBool32 DepthBoundsTestEnable)
+	{
+		FWrapLayer::CmdSetDepthBoundsTestEnableEXT(VK_RESULT_MAX_ENUM, CommandBuffer, DepthBoundsTestEnable);
+		VULKANAPINAMESPACE::vkCmdSetDepthBoundsTestEnableEXT(CommandBuffer, DepthBoundsTestEnable);
+		FWrapLayer::CmdSetDepthBoundsTestEnableEXT(VK_SUCCESS, CommandBuffer, DepthBoundsTestEnable);
+	}
+
+	static FORCEINLINE_DEBUGGABLE void vkCmdSetDepthClampEnableEXT(VkCommandBuffer CommandBuffer, VkBool32 DepthClampEnable)
+	{
+		FWrapLayer::CmdSetDepthClampEnableEXT(VK_RESULT_MAX_ENUM, CommandBuffer, DepthClampEnable);
+		VULKANAPINAMESPACE::vkCmdSetDepthClampEnableEXT(CommandBuffer, DepthClampEnable);
+		FWrapLayer::CmdSetDepthClampEnableEXT(VK_SUCCESS, CommandBuffer, DepthClampEnable);
+	}
+
+	static FORCEINLINE_DEBUGGABLE void vkCmdSetDepthCompareOpEXT(VkCommandBuffer CommandBuffer, VkCompareOp DepthCompareOp)
+	{
+		FWrapLayer::CmdSetDepthCompareOpEXT(VK_RESULT_MAX_ENUM, CommandBuffer, DepthCompareOp);
+		VULKANAPINAMESPACE::vkCmdSetDepthCompareOpEXT(CommandBuffer, DepthCompareOp);
+		FWrapLayer::CmdSetDepthCompareOpEXT(VK_SUCCESS, CommandBuffer, DepthCompareOp);
+	}
+
+	static FORCEINLINE_DEBUGGABLE void vkCmdSetDepthTestEnableEXT(VkCommandBuffer CommandBuffer, VkBool32 DepthTestEnable)
+	{
+		FWrapLayer::CmdSetDepthTestEnableEXT(VK_RESULT_MAX_ENUM, CommandBuffer, DepthTestEnable);
+		VULKANAPINAMESPACE::vkCmdSetDepthTestEnableEXT(CommandBuffer, DepthTestEnable);
+		FWrapLayer::CmdSetDepthTestEnableEXT(VK_SUCCESS, CommandBuffer, DepthTestEnable);
+	}
+
+	static FORCEINLINE_DEBUGGABLE void vkCmdSetDepthWriteEnableEXT(VkCommandBuffer CommandBuffer, VkBool32 DepthWriteEnable)
+	{
+		FWrapLayer::CmdSetDepthWriteEnableEXT(VK_RESULT_MAX_ENUM, CommandBuffer, DepthWriteEnable);
+		VULKANAPINAMESPACE::vkCmdSetDepthWriteEnableEXT(CommandBuffer, DepthWriteEnable);
+		FWrapLayer::CmdSetDepthWriteEnableEXT(VK_SUCCESS, CommandBuffer, DepthWriteEnable);
+	}
+
+	static FORCEINLINE_DEBUGGABLE void vkCmdSetRasterizerDiscardEnableEXT(VkCommandBuffer CommandBuffer, VkBool32 RasterizerDiscardEnable)
+	{
+		FWrapLayer::CmdSetRasterizerDiscardEnableEXT(VK_RESULT_MAX_ENUM, CommandBuffer, RasterizerDiscardEnable);
+		VULKANAPINAMESPACE::vkCmdSetRasterizerDiscardEnableEXT(CommandBuffer, RasterizerDiscardEnable);
+		FWrapLayer::CmdSetRasterizerDiscardEnableEXT(VK_SUCCESS, CommandBuffer, RasterizerDiscardEnable);
+	}
+
+
 #if VULKAN_ENABLE_IMAGE_TRACKING_LAYER
 	void BindDebugLabelName(VkImage Image, const TCHAR* Name);
 #endif
diff --git a/Engine/Source/Runtime/VulkanRHI/Private/VulkanDebug.cpp b/Engine/Source/Runtime/VulkanRHI/Private/VulkanDebug.cpp
--- a/Engine/Source/Runtime/VulkanRHI/Private/VulkanDebug.cpp
+++ b/Engine/Source/Runtime/VulkanRHI/Private/VulkanDebug.cpp
@@ -4705,6 +4705,259 @@ void FWrapLayer::SignalSemaphoreKHR(VkResult Result, VkDevice Device, const VkSe
 	}
 }
 
+void FWrapLayer::CreateShadersEXT(VkResult Result, VkDevice Device, uint32_t CreateInfoCount, const VkShaderCreateInfoEXT* CreateInfos, const VkAllocationCallbacks* Allocator, VkShaderEXT* Shaders)
+{
+	if (Result == VK_RESULT_MAX_ENUM)
+	{
+#if VULKAN_ENABLE_DUMP_LAYER
+		PrintfBeginResult(FString::Printf(TEXT("CreateShadersEXT()")));
+#endif
+	}
+}
+
+void FWrapLayer::DestroyShaderEXT(VkResult Result, VkDevice Device, VkShaderEXT Shader, const VkAllocationCallbacks* Allocator)
+{
+	if (Result == VK_RESULT_MAX_ENUM)
+	{
+#if VULKAN_ENABLE_DUMP_LAYER
+		PrintfBeginResult(FString::Printf(TEXT("DestroyShaderEXT()")));
+#endif
+	}
+}
+
+void FWrapLayer::GetShaderBinaryDataEXT(VkResult Result, VkDevice Device, VkShaderEXT Shader, size_t* DataSize, void* Data)
+{
+	if (Result == VK_RESULT_MAX_ENUM)
+	{
+#if VULKAN_ENABLE_DUMP_LAYER
+		PrintfBeginResult(FString::Printf(TEXT("GetShaderBinaryDataEXT()")));
+#endif
+	}
+}
+
+void FWrapLayer::CmdBindShadersEXT(VkResult Result, VkCommandBuffer CommandBuffer, uint32_t StageCount, const VkShaderStageFlagBits* Stages, const VkShaderEXT* Shaders)
+{
+	if (Result == VK_RESULT_MAX_ENUM)
+	{
+#if VULKAN_ENABLE_DUMP_LAYER
+		PrintfBeginResult(FString::Printf(TEXT("CmdBindShadersEXT()")));
+#endif
+	}
+}
+
+void FWrapLayer::CmdBindVertexBuffers2EXT(VkResult Result, VkCommandBuffer CommandBuffer, uint32_t FirstBinding, uint32_t BindingCount, const VkBuffer* Buffers, const VkDeviceSize* Offsets, const VkDeviceSize* Sizes, const VkDeviceSize* Strides)
+{
+	if (Result == VK_RESULT_MAX_ENUM)
+	{
+#if VULKAN_ENABLE_DUMP_LAYER
+		PrintfBeginResult(FString::Printf(TEXT("CmdBindVertexBuffers2EXT()")));
+#endif
+	}
+}
+
+void FWrapLayer::CmdSetCullModeEXT(VkResult Result, VkCommandBuffer CommandBuffer, VkCullModeFlags CullMode)
+{
+	if (Result == VK_RESULT_MAX_ENUM)
+	{
+#if VULKAN_ENABLE_DUMP_LAYER
+		PrintfBeginResult(FString::Printf(TEXT("CmdSetCullModeEXT()")));
+#endif
+	}
+}
+
+void FWrapLayer::CmdSetColorBlendEnableEXT(VkResult Result, VkCommandBuffer CommandBuffer, uint32_t FirstAttachment, uint32_t AttachmentCount, const VkBool32* ColorBlendEnables)
+{
+	if (Result == VK_RESULT_MAX_ENUM)
+	{
+#if VULKAN_ENABLE_DUMP_LAYER
+		PrintfBeginResult(FString::Printf(TEXT("CmdSetColorBlendEnableEXT()")));
+#endif
+	}
+}
+
+void FWrapLayer::CmdSetColorBlendEquationEXT(VkResult Result, VkCommandBuffer CommandBuffer, uint32_t FirstAttachment, uint32_t AttachmentCount, const VkColorBlendEquationEXT* ColorBlendEquations)
+{
+	if (Result == VK_RESULT_MAX_ENUM)
+	{
+#if VULKAN_ENABLE_DUMP_LAYER
+		PrintfBeginResult(FString::Printf(TEXT("CmdSetColorBlendEquationEXT()")));
+#endif
+	}
+}
+
+void FWrapLayer::CmdSetColorWriteMaskEXT(VkResult Result, VkCommandBuffer CommandBuffer, uint32_t FirstAttachment, uint32_t AttachmentCount, const VkColorComponentFlags* ColorWriteMasks)
+{
+	if (Result == VK_RESULT_MAX_ENUM)
+	{
+#if VULKAN_ENABLE_DUMP_LAYER
+		PrintfBeginResult(FString::Printf(TEXT("CmdSetColorWriteMaskEXT()")));
+#endif
+	}
+}
+
+void FWrapLayer::CmdSetFrontFaceEXT(VkResult Result, VkCommandBuffer CommandBuffer, VkFrontFace FrontFace)
+{
+	if (Result == VK_RESULT_MAX_ENUM)
+	{
+#if VULKAN_ENABLE_DUMP_LAYER
+		PrintfBeginResult(FString::Printf(TEXT("CmdSetFrontFaceEXT()")));
+#endif
+	}
+}
+
+void FWrapLayer::CmdSetPolygonModeEXT(VkResult Result, VkCommandBuffer CommandBuffer, VkPolygonMode PolygonMode)
+{
+	if (Result == VK_RESULT_MAX_ENUM)
+	{
+#if VULKAN_ENABLE_DUMP_LAYER
+		PrintfBeginResult(FString::Printf(TEXT("CmdSetPolygonModeEXT()")));
+#endif
+	}
+}
+
+void FWrapLayer::CmdSetPrimitiveTopologyEXT(VkResult Result, VkCommandBuffer CommandBuffer, VkPrimitiveTopology PrimitiveTopology)
+{
+	if (Result == VK_RESULT_MAX_ENUM)
+	{
+#if VULKAN_ENABLE_DUMP_LAYER
+		PrintfBeginResult(FString::Printf(TEXT("CmdSetPrimitiveTopologyEXT()")));
+#endif
+	}
+}
+
+void FWrapLayer::CmdSetRasterizationSamplesEXT(VkResult Result, VkCommandBuffer CommandBuffer, VkSampleCountFlagBits RasterizationSamples)
+{
+	if (Result == VK_RESULT_MAX_ENUM)
+	{
+#if VULKAN_ENABLE_DUMP_LAYER
+		PrintfBeginResult(FString::Printf(TEXT("CmdSetRasterizationSamplesEXT()")));
+#endif
+	}
+}
+
+void FWrapLayer::CmdSetScissorWithCountEXT(VkResult Result, VkCommandBuffer CommandBuffer, uint32_t ScissorCount, const VkRect2D* Scissors)
+{
+	if (Result == VK_RESULT_MAX_ENUM)
+	{
+#if VULKAN_ENABLE_DUMP_LAYER
+		PrintfBeginResult(FString::Printf(TEXT("CmdSetScissorWithCountEXT()")));
+#endif
+	}
+}
+
+void FWrapLayer::CmdSetStencilOpEXT(VkResult Result, VkCommandBuffer CommandBuffer, VkStencilFaceFlags FaceMask, VkStencilOp FailOp, VkStencilOp passOp, VkStencilOp depthFailOp, VkCompareOp compareOp)
+{
+	if (Result == VK_RESULT_MAX_ENUM)
+	{
+#if VULKAN_ENABLE_DUMP_LAYER
+		PrintfBeginResult(FString::Printf(TEXT("CmdSetStencilOpEXT()")));
+#endif
+	}
+}
+
+void FWrapLayer::CmdSetStencilTestEnableEXT(VkResult Result, VkCommandBuffer CommandBuffer, VkBool32 StencilTestEnable)
+{
+	if (Result == VK_RESULT_MAX_ENUM)
+	{
+#if VULKAN_ENABLE_DUMP_LAYER
+		PrintfBeginResult(FString::Printf(TEXT("CmdSetStencilTestEnableEXT()")));
+#endif
+	}
+}
+
+void FWrapLayer::CmdSetVertexInputEXT(VkResult Result, VkCommandBuffer CommandBuffer, uint32_t VertexBindingDescriptionCount, const VkVertexInputBindingDescription2EXT* pVertexBindingDescriptions, uint32_t vertexAttributeDescriptionCount, const VkVertexInputAttributeDescription2EXT* pVertexAttributeDescriptions)
+{
+	if (Result == VK_RESULT_MAX_ENUM)
+	{
+#if VULKAN_ENABLE_DUMP_LAYER
+		PrintfBeginResult(FString::Printf(TEXT("CmdSetVertexInputEXT()")));
+#endif
+	}
+}
+
+void FWrapLayer::CmdSetViewportWithCountEXT(VkResult Result, VkCommandBuffer CommandBuffer, uint32_t ViewportCount, const VkViewport* pViewports)
+{
+	if (Result == VK_RESULT_MAX_ENUM)
+	{
+#if VULKAN_ENABLE_DUMP_LAYER
+		PrintfBeginResult(FString::Printf(TEXT("CmdSetViewportWithCountEXT()")));
+#endif
+	}
+}
+
+void FWrapLayer::CmdSetDepthBiasEnableEXT(VkResult Result, VkCommandBuffer CommandBuffer, VkBool32 DepthBiasEnable)
+{
+	if (Result == VK_RESULT_MAX_ENUM)
+	{
+#if VULKAN_ENABLE_DUMP_LAYER
+		PrintfBeginResult(FString::Printf(TEXT("CmdSetDepthBiasEnableEXT()")));
+#endif
+	}
+}
+
+void FWrapLayer::CmdSetDepthBoundsTestEnableEXT(VkResult Result, VkCommandBuffer CommandBuffer, VkBool32 DepthBoundsTestEnable)
+{
+	if (Result == VK_RESULT_MAX_ENUM)
+	{
+#if VULKAN_ENABLE_DUMP_LAYER
+		PrintfBeginResult(FString::Printf(TEXT("CmdSetDepthBoundsTestEnableEXT()")));
+#endif
+	}
+}
+
+void FWrapLayer::CmdSetDepthClampEnableEXT(VkResult Result, VkCommandBuffer CommandBuffer, VkBool32 DepthClampEnable)
+{
+	if (Result == VK_RESULT_MAX_ENUM)
+	{
+#if VULKAN_ENABLE_DUMP_LAYER
+		PrintfBeginResult(FString::Printf(TEXT("CmdSetDepthClampEnableEXT()")));
+#endif
+	}
+}
+
+void FWrapLayer::CmdSetDepthCompareOpEXT(VkResult Result, VkCommandBuffer CommandBuffer, VkCompareOp DepthCompareOp)
+{
+	if (Result == VK_RESULT_MAX_ENUM)
+	{
+#if VULKAN_ENABLE_DUMP_LAYER
+		PrintfBeginResult(FString::Printf(TEXT("CmdSetDepthCompareOpEXT()")));
+#endif
+	}
+}
+
+void FWrapLayer::CmdSetDepthTestEnableEXT(VkResult Result, VkCommandBuffer CommandBuffer, VkBool32 DepthTestEnable)
+{
+	if (Result == VK_RESULT_MAX_ENUM)
+	{
+#if VULKAN_ENABLE_DUMP_LAYER
+		PrintfBeginResult(FString::Printf(TEXT("CmdSetDepthTestEnableEXT()")));
+#endif
+	}
+}
+
+void FWrapLayer::CmdSetDepthWriteEnableEXT(VkResult Result, VkCommandBuffer CommandBuffer, VkBool32 DepthWriteEnable)
+{
+	if (Result == VK_RESULT_MAX_ENUM)
+	{
+#if VULKAN_ENABLE_DUMP_LAYER
+		PrintfBeginResult(FString::Printf(TEXT("CmdSetDepthWriteEnableEXT()")));
+#endif
+	}
+}
+
+void FWrapLayer::CmdSetRasterizerDiscardEnableEXT(VkResult Result, VkCommandBuffer CommandBuffer, VkBool32 RasterizerDiscardEnable)
+{
+	if (Result == VK_RESULT_MAX_ENUM)
+	{
+#if VULKAN_ENABLE_DUMP_LAYER
+		PrintfBeginResult(FString::Printf(TEXT("CmdSetRasterizerDiscardEnableEXT()")));
+#endif
+	}
+}
+
+
+
+
 #if VULKAN_ENABLE_IMAGE_TRACKING_LAYER
 namespace VulkanRHI
 {
diff --git a/Engine/Source/Runtime/VulkanRHI/Private/VulkanDevice.h b/Engine/Source/Runtime/VulkanRHI/Private/VulkanDevice.h
--- a/Engine/Source/Runtime/VulkanRHI/Private/VulkanDevice.h
+++ b/Engine/Source/Runtime/VulkanRHI/Private/VulkanDevice.h
@@ -58,12 +58,12 @@ struct FOptionalVulkanDeviceExtensions
 			uint64 HasEXTCalibratedTimestamps : 1;
 			uint64 HasEXTDescriptorBuffer : 1;
 			uint64 HasEXTDeviceFault : 1;
-			uint64 HasEXTLoadStoreOpNone : 1;
 			uint64 HasEXTMeshShader : 1;
 			uint64 HasEXTToolingInfo : 1;
 			uint64 HasEXTImageCompressionControl : 1;
 			uint64 HasEXTMutableDescriptorType : 1;
 			uint64 HasKHRMaintenance7 : 1;
+			uint64 HasEXTShaderObject : 1;
 
 			// Vendor specific
 			uint64 HasAMDBufferMarker : 1;
@@ -100,6 +100,10 @@ struct FOptionalVulkanDeviceExtensions
 			uint64 HasKHRDynamicRendering : 1;
 			uint64 HasEXTSubgroupSizeControl : 1;
 			uint64 HasEXTPipelineCreationCacheControl : 1;
+			uint64 HasEXTExtendedDynamicStates : 1;  // (includes all three extensions)
+
+			// Promoted to 1.4
+			uint64 HasEXTLoadStoreOpNone : 1;
 		};
 		uint64 Packed;
 	};
diff --git a/Engine/Source/Runtime/VulkanRHI/Private/VulkanExtensions.cpp b/Engine/Source/Runtime/VulkanRHI/Private/VulkanExtensions.cpp
--- a/Engine/Source/Runtime/VulkanRHI/Private/VulkanExtensions.cpp
+++ b/Engine/Source/Runtime/VulkanRHI/Private/VulkanExtensions.cpp
@@ -82,6 +82,24 @@ TAutoConsoleVariable<int32> GVulkanAllowDynamicRenderingCVar(
 	ECVF_ReadOnly
 );
 
+TAutoConsoleVariable<int32> GVulkanAllowDynamicStatesCVar(
+	TEXT("r.Vulkan.AllowDynamicStates"),
+	0,
+	TEXT("Enables the use of VK_EXT_extended_dynamic_state1,2,3 to reduce amount of unique pipelines. \n")
+	TEXT("0: Do not enable support for all the dynamic state extensions (default)\n")
+	TEXT("1: Enable all the dynamic state extensions"),
+	ECVF_ReadOnly
+);
+
+TAutoConsoleVariable<int32> GVulkanAllowShaderObjectCVar(
+	TEXT("r.Vulkan.AllowShaderObject"),
+	0,
+	TEXT("Enables the use of VK_EXT_shader_object to replace graphic pipelines. \n")
+	TEXT("0: Do not enable support for shader objects (default)\n")
+	TEXT("1: Enable shader objects"),
+	ECVF_ReadOnly
+);
+
 TAutoConsoleVariable<int32> GVulkanVariableRateShadingFormatCVar(
 	TEXT("r.Vulkan.VRSFormat"),
 	0,
@@ -210,7 +228,7 @@ class FVulkanKHRMaintenance5Extension : public FVulkanDeviceExtension
 public:
 
 	FVulkanKHRMaintenance5Extension(FVulkanDevice* InDevice)
-		: FVulkanDeviceExtension(InDevice, VK_KHR_MAINTENANCE_5_EXTENSION_NAME, VULKAN_EXTENSION_ENABLED)
+		: FVulkanDeviceExtension(InDevice, VK_KHR_MAINTENANCE_5_EXTENSION_NAME, VULKAN_EXTENSION_ENABLED, VK_API_VERSION_1_4)
 	{}
 
 	virtual void PrePhysicalDeviceFeatures(VkPhysicalDeviceFeatures2KHR& PhysicalDeviceFeatures2) override final
@@ -1935,6 +1953,167 @@ public:
 	VkPhysicalDeviceDynamicRenderingUnusedAttachmentsFeaturesEXT DynamicRenderingUnusedAttachmentsFeatures;
 };
 
+// ***** VK_EXT_extended_dynamic_state
+class FVulkanEXTExtendedDynamicStateExtension : public FVulkanDeviceExtension
+{
+public:
+
+	FVulkanEXTExtendedDynamicStateExtension(FVulkanDevice* InDevice)
+		: FVulkanDeviceExtension(InDevice, VK_EXT_EXTENDED_DYNAMIC_STATE_EXTENSION_NAME, VULKAN_EXTENSION_ENABLED, VK_API_VERSION_1_3)
+	{
+		bEnabledInCode = bEnabledInCode && (GVulkanAllowDynamicStatesCVar.GetValueOnAnyThread() != 0);
+	}
+
+	virtual void PrePhysicalDeviceFeatures(VkPhysicalDeviceFeatures2KHR& PhysicalDeviceFeatures2) override final
+	{
+		ZeroVulkanStruct(ExtendedDynamicStateFeatures, VK_STRUCTURE_TYPE_PHYSICAL_DEVICE_EXTENDED_DYNAMIC_STATE_FEATURES_EXT);
+		AddToPNext(PhysicalDeviceFeatures2, ExtendedDynamicStateFeatures);
+	}
+
+	virtual void PostPhysicalDeviceFeatures(FOptionalVulkanDeviceExtensions& ExtensionFlags) override final
+	{
+		bRequirementsPassed = (ExtendedDynamicStateFeatures.extendedDynamicState == VK_TRUE);
+		ExtensionFlags.HasEXTExtendedDynamicStates = bRequirementsPassed;
+	}
+
+	virtual void PreCreateDevice(VkDeviceCreateInfo& DeviceCreateInfo) override final
+	{
+		if (bRequirementsPassed)
+		{
+			AddToPNext(DeviceCreateInfo, ExtendedDynamicStateFeatures);
+		}
+	}
+
+	VkPhysicalDeviceExtendedDynamicStateFeaturesEXT ExtendedDynamicStateFeatures;
+};
+
+// ***** VK_EXT_extended_dynamic_state2
+class FVulkanEXTExtendedDynamicState2Extension : public FVulkanDeviceExtension
+{
+public:
+
+	FVulkanEXTExtendedDynamicState2Extension(FVulkanDevice* InDevice)
+		: FVulkanDeviceExtension(InDevice, VK_EXT_EXTENDED_DYNAMIC_STATE_2_EXTENSION_NAME, VULKAN_EXTENSION_ENABLED, VK_API_VERSION_1_3)
+	{
+		bEnabledInCode = bEnabledInCode && (GVulkanAllowDynamicStatesCVar.GetValueOnAnyThread() != 0);
+	}
+
+	virtual void PrePhysicalDeviceFeatures(VkPhysicalDeviceFeatures2KHR& PhysicalDeviceFeatures2) override final
+	{
+		ZeroVulkanStruct(ExtendedDynamicState2Features, VK_STRUCTURE_TYPE_PHYSICAL_DEVICE_EXTENDED_DYNAMIC_STATE_2_FEATURES_EXT);
+		AddToPNext(PhysicalDeviceFeatures2, ExtendedDynamicState2Features);
+	}
+
+	virtual void PostPhysicalDeviceFeatures(FOptionalVulkanDeviceExtensions& ExtensionFlags) override final
+	{
+		if (ExtensionFlags.HasEXTExtendedDynamicStates)
+		{
+			bRequirementsPassed = (ExtendedDynamicState2Features.extendedDynamicState2 == VK_TRUE);
+			ExtensionFlags.HasEXTExtendedDynamicStates = bRequirementsPassed;
+		}
+		else
+		{
+			bRequirementsPassed = false;
+		}
+	}
+
+	virtual void PreCreateDevice(VkDeviceCreateInfo& DeviceCreateInfo) override final
+	{
+		if (bRequirementsPassed)
+		{
+			AddToPNext(DeviceCreateInfo, ExtendedDynamicState2Features);
+		}
+	}
+
+	VkPhysicalDeviceExtendedDynamicState2FeaturesEXT ExtendedDynamicState2Features;
+};
+
+// ***** VK_EXT_extended_dynamic_state3
+class FVulkanEXTExtendedDynamicState3Extension : public FVulkanDeviceExtension
+{
+public:
+
+	FVulkanEXTExtendedDynamicState3Extension(FVulkanDevice* InDevice)
+		: FVulkanDeviceExtension(InDevice, VK_EXT_EXTENDED_DYNAMIC_STATE_3_EXTENSION_NAME, VULKAN_EXTENSION_ENABLED)
+	{
+		bEnabledInCode = bEnabledInCode && (GVulkanAllowDynamicStatesCVar.GetValueOnAnyThread() != 0);
+	}
+
+	virtual void PrePhysicalDeviceFeatures(VkPhysicalDeviceFeatures2KHR& PhysicalDeviceFeatures2) override final
+	{
+		ZeroVulkanStruct(ExtendedDynamicState3Features, VK_STRUCTURE_TYPE_PHYSICAL_DEVICE_EXTENDED_DYNAMIC_STATE_3_FEATURES_EXT);
+		AddToPNext(PhysicalDeviceFeatures2, ExtendedDynamicState3Features);
+	}
+
+	virtual void PostPhysicalDeviceFeatures(FOptionalVulkanDeviceExtensions& ExtensionFlags) override final
+	{
+		if (ExtensionFlags.HasEXTExtendedDynamicStates)
+		{
+			bRequirementsPassed = 
+				(ExtendedDynamicState3Features.extendedDynamicState3PolygonMode == VK_TRUE) &&
+				(ExtendedDynamicState3Features.extendedDynamicState3ColorBlendEnable == VK_TRUE) &&
+				(ExtendedDynamicState3Features.extendedDynamicState3ColorBlendEquation == VK_TRUE) &&
+				(ExtendedDynamicState3Features.extendedDynamicState3ColorWriteMask == VK_TRUE) &&
+				(ExtendedDynamicState3Features.extendedDynamicState3RasterizationSamples == VK_TRUE);
+			ExtensionFlags.HasEXTExtendedDynamicStates = bRequirementsPassed;
+		}
+		else
+		{
+			bRequirementsPassed = false;
+		}
+	}
+
+	virtual void PreCreateDevice(VkDeviceCreateInfo& DeviceCreateInfo) override final
+	{
+		if (bRequirementsPassed)
+		{
+			AddToPNext(DeviceCreateInfo, ExtendedDynamicState3Features);
+		}
+	}
+
+	VkPhysicalDeviceExtendedDynamicState3FeaturesEXT ExtendedDynamicState3Features;
+};
+
+// ***** VK_EXT_shader_object
+class FVulkanEXTShaderObjectExtension : public FVulkanDeviceExtension
+{
+public:
+
+	FVulkanEXTShaderObjectExtension(FVulkanDevice* InDevice)
+		: FVulkanDeviceExtension(InDevice, VK_EXT_SHADER_OBJECT_EXTENSION_NAME, VULKAN_EXTENSION_ENABLED)
+	{
+		bEnabledInCode = bEnabledInCode && 
+			(GVulkanAllowDynamicStatesCVar.GetValueOnAnyThread() != 0) &&
+			(GVulkanAllowShaderObjectCVar.GetValueOnAnyThread() != 0);
+	}
+
+	virtual void PrePhysicalDeviceFeatures(VkPhysicalDeviceFeatures2KHR& PhysicalDeviceFeatures2) override final
+	{
+		ZeroVulkanStruct(ShaderObjectFeatures, VK_STRUCTURE_TYPE_PHYSICAL_DEVICE_SHADER_OBJECT_FEATURES_EXT);
+		AddToPNext(PhysicalDeviceFeatures2, ShaderObjectFeatures);
+	}
+
+	virtual void PostPhysicalDeviceFeatures(FOptionalVulkanDeviceExtensions& ExtensionFlags) override final
+	{
+		bRequirementsPassed = (ShaderObjectFeatures.shaderObject == VK_TRUE);
+		ExtensionFlags.HasEXTShaderObject = bRequirementsPassed;
+	}
+
+	virtual void PreCreateDevice(VkDeviceCreateInfo& DeviceCreateInfo) override final
+	{
+		if (bRequirementsPassed)
+		{
+			AddToPNext(DeviceCreateInfo, ShaderObjectFeatures);
+		}
+	}
+
+	VkPhysicalDeviceShaderObjectFeaturesEXT ShaderObjectFeatures;
+};
+
+
+
+
+
 template <typename ExtensionType>
 static void FlagExtensionSupport(const TArray<VkExtensionProperties>& ExtensionProperties, TArray<TUniquePtr<ExtensionType>>& UEExtensions, uint32 ApiVersion, const TCHAR* ExtensionTypeName)
 {
@@ -1984,7 +2163,7 @@ FVulkanDeviceExtensionArray FVulkanDeviceExtension::GetUESupportedDeviceExtensio
 	ADD_SIMPLE_EXTENSION(VK_KHR_SHADER_FLOAT_CONTROLS_EXTENSION_NAME,	 VULKAN_EXTENSION_ENABLED,             VK_API_VERSION_1_2,            DEVICE_EXT_FLAG_SETTER(HasShaderFloatControls));
 	ADD_SIMPLE_EXTENSION(VK_KHR_IMAGE_FORMAT_LIST_EXTENSION_NAME,        VULKAN_EXTENSION_ENABLED,             VK_API_VERSION_1_2,            DEVICE_EXT_FLAG_SETTER(HasKHRImageFormatList));
 	ADD_SIMPLE_EXTENSION(VK_EXT_VALIDATION_CACHE_EXTENSION_NAME,         VULKAN_SUPPORTS_VALIDATION_CACHE,     VULKAN_EXTENSION_NOT_PROMOTED, DEVICE_EXT_FLAG_SETTER(HasEXTValidationCache));
-	ADD_SIMPLE_EXTENSION(VK_EXT_LOAD_STORE_OP_NONE_EXTENSION_NAME,       VULKAN_EXTENSION_ENABLED,             VULKAN_EXTENSION_NOT_PROMOTED, DEVICE_EXT_FLAG_SETTER(HasEXTLoadStoreOpNone));
+	ADD_SIMPLE_EXTENSION(VK_EXT_LOAD_STORE_OP_NONE_EXTENSION_NAME,       VULKAN_EXTENSION_ENABLED,             VK_API_VERSION_1_4,            DEVICE_EXT_FLAG_SETTER(HasEXTLoadStoreOpNone));
 	ADD_SIMPLE_EXTENSION(VK_QCOM_RENDER_PASS_SHADER_RESOLVE_EXTENSION_NAME, VULKAN_SUPPORTS_QCOM_RENDERPASS_SHADER_RESOLVE, VULKAN_EXTENSION_NOT_PROMOTED, DEVICE_EXT_FLAG_SETTER(HasQcomRenderPassShaderResolve));
 	ADD_SIMPLE_EXTENSION(VK_EXT_TOOLING_INFO_EXTENSION_NAME, VULKAN_HAS_DEBUGGING_ENABLED || VULKAN_ENABLE_DRAW_MARKERS, VK_API_VERSION_1_3,  DEVICE_EXT_FLAG_SETTER(HasEXTToolingInfo));
 	ADD_SIMPLE_EXTENSION(VK_KHR_PIPELINE_LIBRARY_EXTENSION_NAME,         VULKAN_EXTENSION_ENABLED,             VULKAN_EXTENSION_NOT_PROMOTED, DEVICE_EXT_FLAG_SETTER(HasKHRPipelineLibrary));
@@ -2031,6 +2210,11 @@ FVulkanDeviceExtensionArray FVulkanDeviceExtension::GetUESupportedDeviceExtensio
 	ADD_CUSTOM_EXTENSION(FVulkanKHRDynamicRenderingExtension);			// must be kept BEFORE DynamicRenderingLocalRead!
 	ADD_CUSTOM_EXTENSION(FVulkanKHRDynamicRenderingLocalReadExtension);	// must be kept AFTER DynamicRendering!
 	ADD_CUSTOM_EXTENSION(FVulkanKHRDynamicRenderingUnusedAttachmentsExtension);	// must be kept AFTER DynamicRendering!
+	ADD_CUSTOM_EXTENSION(FVulkanKHRDynamicRenderingUnusedAttachmentsExtension);	// must be kept AFTER DynamicRendering!
+	ADD_CUSTOM_EXTENSION(FVulkanEXTExtendedDynamicStateExtension);	// extended dynamic state extension must be kept in order (1,2,3)
+	ADD_CUSTOM_EXTENSION(FVulkanEXTExtendedDynamicState2Extension); // extended dynamic state extension must be kept in order (1,2,3)
+	ADD_CUSTOM_EXTENSION(FVulkanEXTExtendedDynamicState3Extension); // extended dynamic state extension must be kept in order (1,2,3)
+	ADD_CUSTOM_EXTENSION(FVulkanEXTShaderObjectExtension);
 
 	// Needed for Raytracing
 	ADD_CUSTOM_EXTENSION(FVulkanKHRBufferDeviceAddressExtension);
diff --git a/Engine/Source/Runtime/VulkanRHI/Private/VulkanLoader.h b/Engine/Source/Runtime/VulkanRHI/Private/VulkanLoader.h
--- a/Engine/Source/Runtime/VulkanRHI/Private/VulkanLoader.h
+++ b/Engine/Source/Runtime/VulkanRHI/Private/VulkanLoader.h
@@ -241,7 +241,32 @@
 	EnumMacro(PFN_vkGetImageSubresourceLayout2KHR, vkGetImageSubresourceLayout2KHR) \
 	EnumMacro(PFN_vkGetSemaphoreCounterValue, vkGetSemaphoreCounterValueKHR) \
 	EnumMacro(PFN_vkWaitSemaphores, vkWaitSemaphoresKHR) \
-	EnumMacro(PFN_vkSignalSemaphore, vkSignalSemaphoreKHR)
+	EnumMacro(PFN_vkSignalSemaphore, vkSignalSemaphoreKHR) \
+	EnumMacro(PFN_vkCreateShadersEXT, vkCreateShadersEXT) \
+	EnumMacro(PFN_vkDestroyShaderEXT, vkDestroyShaderEXT) \
+	EnumMacro(PFN_vkGetShaderBinaryDataEXT, vkGetShaderBinaryDataEXT) \
+	EnumMacro(PFN_vkCmdBindShadersEXT, vkCmdBindShadersEXT) \
+	EnumMacro(PFN_vkCmdBindVertexBuffers2EXT, vkCmdBindVertexBuffers2EXT) \
+	EnumMacro(PFN_vkCmdSetCullModeEXT, vkCmdSetCullModeEXT) \
+	EnumMacro(PFN_vkCmdSetColorBlendEnableEXT, vkCmdSetColorBlendEnableEXT) \
+	EnumMacro(PFN_vkCmdSetColorBlendEquationEXT, vkCmdSetColorBlendEquationEXT) \
+	EnumMacro(PFN_vkCmdSetColorWriteMaskEXT, vkCmdSetColorWriteMaskEXT) \
+	EnumMacro(PFN_vkCmdSetFrontFaceEXT, vkCmdSetFrontFaceEXT) \
+	EnumMacro(PFN_vkCmdSetPolygonModeEXT, vkCmdSetPolygonModeEXT) \
+	EnumMacro(PFN_vkCmdSetPrimitiveTopologyEXT, vkCmdSetPrimitiveTopologyEXT) \
+	EnumMacro(PFN_vkCmdSetRasterizationSamplesEXT, vkCmdSetRasterizationSamplesEXT) \
+	EnumMacro(PFN_vkCmdSetScissorWithCountEXT, vkCmdSetScissorWithCountEXT) \
+	EnumMacro(PFN_vkCmdSetStencilOpEXT, vkCmdSetStencilOpEXT) \
+	EnumMacro(PFN_vkCmdSetStencilTestEnableEXT, vkCmdSetStencilTestEnableEXT) \
+	EnumMacro(PFN_vkCmdSetVertexInputEXT, vkCmdSetVertexInputEXT) \
+	EnumMacro(PFN_vkCmdSetViewportWithCountEXT, vkCmdSetViewportWithCountEXT) \
+	EnumMacro(PFN_vkCmdSetDepthBiasEnableEXT, vkCmdSetDepthBiasEnableEXT) \
+	EnumMacro(PFN_vkCmdSetDepthBoundsTestEnableEXT, vkCmdSetDepthBoundsTestEnableEXT) \
+	EnumMacro(PFN_vkCmdSetDepthClampEnableEXT, vkCmdSetDepthClampEnableEXT) \
+	EnumMacro(PFN_vkCmdSetDepthCompareOpEXT, vkCmdSetDepthCompareOpEXT) \
+	EnumMacro(PFN_vkCmdSetDepthTestEnableEXT, vkCmdSetDepthTestEnableEXT) \
+	EnumMacro(PFN_vkCmdSetDepthWriteEnableEXT, vkCmdSetDepthWriteEnableEXT) \
+	EnumMacro(PFN_vkCmdSetRasterizerDiscardEnableEXT, vkCmdSetRasterizerDiscardEnableEXT)
 
 
 
