VulkanRHI: Only read ClearColor if LoadOp is Clear.  Fix assert in EngineTests with dynamic rendering.

#rnx

diff --git a/Engine/Source/Runtime/VulkanRHI/Private/VulkanContext.cpp b/Engine/Source/Runtime/VulkanRHI/Private/VulkanContext.cpp
--- a/Engine/Source/Runtime/VulkanRHI/Private/VulkanContext.cpp
+++ b/Engine/Source/Runtime/VulkanRHI/Private/VulkanContext.cpp
@@ -324,11 +324,14 @@ void FVulkanCommandListContext::FillDynamicRenderingInfo(const FRHIRenderPassInf
 			AttachmentInfo.storeOp = RenderTargetStoreActionToVulkan(GetStoreAction(ColorEntry.Action));
 			ensure(!EnumHasAnyFlags(TextureDesc.Flags, TexCreate_Memoryless) || (AttachmentInfo.storeOp == VK_ATTACHMENT_STORE_OP_DONT_CARE));
 
-			const FLinearColor ClearColor = Texture->HasClearValue() ? Texture->GetClearColor() : FLinearColor::Black;
-			AttachmentInfo.clearValue.color.float32[0] = ClearColor.R;
-			AttachmentInfo.clearValue.color.float32[1] = ClearColor.G;
-			AttachmentInfo.clearValue.color.float32[2] = ClearColor.B;
-			AttachmentInfo.clearValue.color.float32[3] = ClearColor.A;
+			if (AttachmentInfo.loadOp == VK_ATTACHMENT_LOAD_OP_CLEAR)
+			{
+				const FLinearColor ClearColor = Texture->HasClearValue() ? Texture->GetClearColor() : FLinearColor::Black;
+				AttachmentInfo.clearValue.color.float32[0] = ClearColor.R;
+				AttachmentInfo.clearValue.color.float32[1] = ClearColor.G;
+				AttachmentInfo.clearValue.color.float32[2] = ClearColor.B;
+				AttachmentInfo.clearValue.color.float32[3] = ClearColor.A;
+			}
 
 			MipIndex = ColorEntry.MipIndex;
 			FVulkanView* View = FVulkanFramebuffer::GetColorRenderTargetViewDesc(Texture, MipIndex, ColorEntry.ArraySlice, InRenderPassInfo.MultiViewCount, NumLayers);
