VulkanRHI: Fix uninitialized variable.  Fix NumLayers in framebuffer since 43607568 (fix texture arrays RHI tests).

#rnx

diff --git a/Engine/Source/Runtime/VulkanRHI/Private/VulkanViewport.cpp b/Engine/Source/Runtime/VulkanRHI/Private/VulkanViewport.cpp
--- a/Engine/Source/Runtime/VulkanRHI/Private/VulkanViewport.cpp
+++ b/Engine/Source/Runtime/VulkanRHI/Private/VulkanViewport.cpp
@@ -295,6 +295,7 @@ FVulkanFramebuffer::FVulkanFramebuffer(FVulkanDevice& Device, const FRHISetRende
 	, NumColorRenderTargets(InRTInfo.NumColorRenderTargets)
 	, NumColorAttachments(0)
 	, DepthStencilRenderTargetImage(VK_NULL_HANDLE)
+	, DepthStencilResolveRenderTargetImage(VK_NULL_HANDLE)
 	, FragmentDensityImage(VK_NULL_HANDLE)
 {
 	FMemory::Memzero(ColorRenderTargetImages);
@@ -426,7 +427,7 @@ FVulkanView* FVulkanFramebuffer::GetColorRenderTargetViewDesc(FVulkanTexture* Te
 
 		// About !RTLayout.GetIsMultiView(), from https://registry.khronos.org/vulkan/specs/1.3-extensions/man/html/VkFramebufferCreateInfo.html: 
 		// If the render pass uses multiview, then layers must be one
-		if (Texture->GetViewType() == VK_IMAGE_VIEW_TYPE_2D_ARRAY && (MultiViewCount != 0))
+		if (Texture->GetViewType() == VK_IMAGE_VIEW_TYPE_2D_ARRAY && (MultiViewCount == 0))
 		{
 			InOutNumLayers = NumArraySlices;
 		}
