Reorganizing a number of Vulkan types to be in their own file(s).

diff --git a/Engine/Source/Runtime/VulkanRHI/Private/VulkanBindlessDescriptorManager.cpp b/Engine/Source/Runtime/VulkanRHI/Private/VulkanBindlessDescriptorManager.cpp
new file mode 100644
--- /dev/null
+++ b/Engine/Source/Runtime/VulkanRHI/Private/VulkanBindlessDescriptorManager.cpp
@@ -0,0 +1,938 @@
+// Copyright Epic Games, Inc. All Rights Reserved.
+
+#include "VulkanBindlessDescriptorManager.h"
+#include "VulkanRHIPrivate.h"
+#include "VulkanContext.h"
+#include "VulkanDevice.h"
+
+int32 GVulkanBindlessMaxSamplerDescriptorCount = 2048;
+static FAutoConsoleVariableRef CVarVulkanBindlessMaxSamplerDescriptorCount(
+	TEXT("r.Vulkan.Bindless.MaxSamplerDescriptorCount"),
+	GVulkanBindlessMaxSamplerDescriptorCount,
+	TEXT("Maximum bindless sampler descriptor count"),
+	ECVF_ReadOnly
+);
+
+int32 GVulkanBindlessMaxSampledImageDescriptorCount = 256 * 1024;
+static FAutoConsoleVariableRef CVarVulkanBindlessMaxSampledImageCount(
+	TEXT("r.Vulkan.Bindless.MaxResourceSampledImageCount"),
+	GVulkanBindlessMaxSampledImageDescriptorCount,
+	TEXT("Maximum bindless Sampled Image descriptor count"),
+	ECVF_ReadOnly
+);
+
+int32 GVulkanBindlessMaxStorageImageDescriptorCount = 64 * 1024;
+static FAutoConsoleVariableRef CVarVulkanBindlessMaxStorageImageCount(
+	TEXT("r.Vulkan.Bindless.MaxResourceStorageImageCount"),
+	GVulkanBindlessMaxStorageImageDescriptorCount,
+	TEXT("Maximum bindless Storage Image descriptor count"),
+	ECVF_ReadOnly
+);
+
+int32 GVulkanBindlessMaxUniformTexelBufferDescriptorCount = 64 * 1024;
+static FAutoConsoleVariableRef CVarVulkanBindlessMaxUniformTexelBufferCount(
+	TEXT("r.Vulkan.Bindless.MaxResourceUniformTexelBufferCount"),
+	GVulkanBindlessMaxUniformTexelBufferDescriptorCount,
+	TEXT("Maximum bindless Uniform Texel Buffer descriptor count"),
+	ECVF_ReadOnly
+);
+
+int32 GVulkanBindlessMaxStorageTexelBufferDescriptorCount = 64 * 1024;
+static FAutoConsoleVariableRef CVarVulkanBindlessMaxStorageTexelBufferCount(
+	TEXT("r.Vulkan.Bindless.MaxResourceStorageTexelBufferCount"),
+	GVulkanBindlessMaxStorageTexelBufferDescriptorCount,
+	TEXT("Maximum bindless Storage Texel Buffer descriptor count"),
+	ECVF_ReadOnly
+);
+
+int32 GVulkanBindlessMaxUniformBufferDescriptorCount = 32 * 1024;
+static FAutoConsoleVariableRef CVarVulkanBindlessMaxUniformBufferCount(
+	TEXT("r.Vulkan.Bindless.MaxResourceUniformBufferCount"),
+	GVulkanBindlessMaxUniformBufferDescriptorCount,
+	TEXT("Maximum bindless Uniform Buffer descriptor count"),
+	ECVF_ReadOnly
+);
+
+int32 GVulkanBindlessMaxStorageBufferDescriptorCount = 64 * 1024;
+static FAutoConsoleVariableRef CVarVulkanBindlessMaxStorageBufferCount(
+	TEXT("r.Vulkan.Bindless.MaxResourceStorageBufferCount"),
+	GVulkanBindlessMaxStorageBufferDescriptorCount,
+	TEXT("Maximum bindless Storage Buffer descriptor count"),
+	ECVF_ReadOnly
+);
+
+int32 GVulkanBindlessMaxAccelerationStructureDescriptorCount = 64 * 1024;
+static FAutoConsoleVariableRef CVarVulkanBindlessMaxAccelerationStructureCount(
+	TEXT("r.Vulkan.Bindless.MaxResourceAccelerationStructureCount"),
+	GVulkanBindlessMaxAccelerationStructureDescriptorCount,
+	TEXT("Maximum bindless Acceleration Structure descriptor count"),
+	ECVF_ReadOnly
+);
+
+int32 GVulkanBindlessBlockSize = 1024 * 1024;
+static FAutoConsoleVariableRef CVarVulkanBindlessBlockSize(
+	TEXT("r.Vulkan.Bindless.BlockSize"),
+	GVulkanBindlessBlockSize,
+	TEXT("Block size to use for single use ub. (default: 1MB)"),
+	ECVF_RenderThreadSafe
+);
+
+
+DECLARE_STATS_GROUP(TEXT("Vulkan Bindless"), STATGROUP_VulkanBindless, STATCAT_Advanced);
+
+DECLARE_DWORD_ACCUMULATOR_STAT_EXTERN(TEXT("Num Peak Descriptor Count"), STAT_VulkanBindlessPeakDescriptorCount, STATGROUP_VulkanBindless, );
+DECLARE_DWORD_ACCUMULATOR_STAT_EXTERN(TEXT("Peak Samplers"), STAT_VulkanBindlessPeakSampler, STATGROUP_VulkanBindless, );
+DECLARE_DWORD_ACCUMULATOR_STAT_EXTERN(TEXT("Peak Sampled Images"), STAT_VulkanBindlessPeakSampledImage, STATGROUP_VulkanBindless, );
+DECLARE_DWORD_ACCUMULATOR_STAT_EXTERN(TEXT("Peak Storage Images"), STAT_VulkanBindlessPeakStorageImage, STATGROUP_VulkanBindless, );
+DECLARE_DWORD_ACCUMULATOR_STAT_EXTERN(TEXT("Peak Uniform Buffers"), STAT_VulkanBindlessPeakUniformBuffer, STATGROUP_VulkanBindless, );
+DECLARE_DWORD_ACCUMULATOR_STAT_EXTERN(TEXT("Peak Storage Buffers"), STAT_VulkanBindlessPeakStorageBuffer, STATGROUP_VulkanBindless, );
+DECLARE_DWORD_ACCUMULATOR_STAT_EXTERN(TEXT("Peak Uniform Texel Buffers"), STAT_VulkanBindlessPeakUniformTexelBuffer, STATGROUP_VulkanBindless, );
+DECLARE_DWORD_ACCUMULATOR_STAT_EXTERN(TEXT("Peak Storage Texel Buffers"), STAT_VulkanBindlessPeakStorageTexelBuffer, STATGROUP_VulkanBindless, );
+DECLARE_DWORD_ACCUMULATOR_STAT_EXTERN(TEXT("Peak Acceleration Structures"), STAT_VulkanBindlessPeakAccelerationStructure, STATGROUP_VulkanBindless, );
+DECLARE_DWORD_COUNTER_STAT_EXTERN(TEXT("Write Per Frame"), STAT_VulkanBindlessWritePerFrame, STATGROUP_VulkanBindless, );
+
+DEFINE_STAT(STAT_VulkanBindlessPeakDescriptorCount);
+DEFINE_STAT(STAT_VulkanBindlessPeakSampler);
+DEFINE_STAT(STAT_VulkanBindlessPeakSampledImage);
+DEFINE_STAT(STAT_VulkanBindlessPeakStorageImage);
+DEFINE_STAT(STAT_VulkanBindlessPeakUniformBuffer);
+DEFINE_STAT(STAT_VulkanBindlessPeakStorageBuffer);
+DEFINE_STAT(STAT_VulkanBindlessPeakUniformTexelBuffer);
+DEFINE_STAT(STAT_VulkanBindlessPeakStorageTexelBuffer);
+DEFINE_STAT(STAT_VulkanBindlessPeakAccelerationStructure);
+DEFINE_STAT(STAT_VulkanBindlessWritePerFrame);
+
+
+static constexpr uint8 GetIndexForDescriptorType(VkDescriptorType DescriptorType)
+{
+	switch (DescriptorType)
+	{
+	case VK_DESCRIPTOR_TYPE_SAMPLER:                    return VulkanBindless::BindlessSamplerSet;
+	case VK_DESCRIPTOR_TYPE_SAMPLED_IMAGE:              return VulkanBindless::BindlessSampledImageSet;
+	case VK_DESCRIPTOR_TYPE_STORAGE_IMAGE:              return VulkanBindless::BindlessStorageImageSet;
+	case VK_DESCRIPTOR_TYPE_UNIFORM_TEXEL_BUFFER:       return VulkanBindless::BindlessUniformTexelBufferSet;
+	case VK_DESCRIPTOR_TYPE_STORAGE_TEXEL_BUFFER:       return VulkanBindless::BindlessStorageTexelBufferSet;
+	case VK_DESCRIPTOR_TYPE_STORAGE_BUFFER:             return VulkanBindless::BindlessStorageBufferSet;
+	case VK_DESCRIPTOR_TYPE_UNIFORM_BUFFER:             return VulkanBindless::BindlessUniformBufferSet;
+	case VK_DESCRIPTOR_TYPE_ACCELERATION_STRUCTURE_KHR: return VulkanBindless::BindlessAccelerationStructureSet;
+	default: checkNoEntry();
+	}
+
+	return VulkanBindless::MaxNumSets;
+}
+
+static constexpr uint8 GetIndexForDescriptorType(ERHIDescriptorType DescriptorType)
+{
+	switch (DescriptorType)
+	{
+	case ERHIDescriptorType::Sampler:               return VulkanBindless::BindlessSamplerSet;
+	case ERHIDescriptorType::TextureSRV:            return VulkanBindless::BindlessSampledImageSet;
+	case ERHIDescriptorType::TextureUAV:            return VulkanBindless::BindlessStorageImageSet;
+	case ERHIDescriptorType::TypedBufferSRV:        return VulkanBindless::BindlessUniformTexelBufferSet;
+	case ERHIDescriptorType::TypedBufferUAV:        return VulkanBindless::BindlessStorageTexelBufferSet;
+	case ERHIDescriptorType::BufferUAV:             return VulkanBindless::BindlessStorageBufferSet;
+	case ERHIDescriptorType::BufferSRV:             return VulkanBindless::BindlessUniformBufferSet;
+	case ERHIDescriptorType::AccelerationStructure: return VulkanBindless::BindlessAccelerationStructureSet;
+	default: checkNoEntry();
+	}
+
+	return VulkanBindless::MaxNumSets;
+}
+
+
+static constexpr ERHIDescriptorType GetUEDescriptorTypeFromVkDescriptorType(VkDescriptorType DescriptorType)
+{
+	switch (DescriptorType)
+	{
+	case VK_DESCRIPTOR_TYPE_SAMPLER:                    return ERHIDescriptorType::Sampler;
+	case VK_DESCRIPTOR_TYPE_SAMPLED_IMAGE:              return ERHIDescriptorType::TextureSRV;
+	case VK_DESCRIPTOR_TYPE_STORAGE_IMAGE:              return ERHIDescriptorType::TextureUAV;
+	case VK_DESCRIPTOR_TYPE_UNIFORM_TEXEL_BUFFER:       return ERHIDescriptorType::TypedBufferSRV;
+	case VK_DESCRIPTOR_TYPE_STORAGE_TEXEL_BUFFER:       return ERHIDescriptorType::TypedBufferUAV;
+	case VK_DESCRIPTOR_TYPE_STORAGE_BUFFER:             return ERHIDescriptorType::BufferUAV;
+	case VK_DESCRIPTOR_TYPE_UNIFORM_BUFFER:             return ERHIDescriptorType::BufferSRV;
+	case VK_DESCRIPTOR_TYPE_ACCELERATION_STRUCTURE_KHR: return ERHIDescriptorType::AccelerationStructure;
+	default: checkNoEntry();
+	}
+
+	return ERHIDescriptorType::Invalid;
+}
+
+
+extern TAutoConsoleVariable<int32> GCVarRobustBufferAccess;
+static inline uint32 GetDescriptorTypeSize(FVulkanDevice& Device, VkDescriptorType DescriptorType)
+{
+	const bool bRobustBufferAccess = (GCVarRobustBufferAccess.GetValueOnAnyThread() > 0);
+	const VkPhysicalDeviceDescriptorBufferPropertiesEXT& DescriptorBufferProperties = Device.GetOptionalExtensionProperties().DescriptorBufferProps;
+
+	switch (DescriptorType)
+	{
+	case VK_DESCRIPTOR_TYPE_SAMPLER:
+		return DescriptorBufferProperties.samplerDescriptorSize;
+	case VK_DESCRIPTOR_TYPE_SAMPLED_IMAGE:
+		return DescriptorBufferProperties.sampledImageDescriptorSize;
+	case VK_DESCRIPTOR_TYPE_STORAGE_IMAGE:
+		return DescriptorBufferProperties.storageImageDescriptorSize;
+	case VK_DESCRIPTOR_TYPE_UNIFORM_TEXEL_BUFFER:
+		return bRobustBufferAccess ? DescriptorBufferProperties.robustUniformTexelBufferDescriptorSize
+			: DescriptorBufferProperties.uniformTexelBufferDescriptorSize;
+	case VK_DESCRIPTOR_TYPE_STORAGE_TEXEL_BUFFER:
+		return bRobustBufferAccess ? DescriptorBufferProperties.robustStorageTexelBufferDescriptorSize
+			: DescriptorBufferProperties.storageTexelBufferDescriptorSize;
+	case VK_DESCRIPTOR_TYPE_UNIFORM_BUFFER:
+		return bRobustBufferAccess ? DescriptorBufferProperties.robustUniformBufferDescriptorSize
+			: DescriptorBufferProperties.uniformBufferDescriptorSize;
+	case VK_DESCRIPTOR_TYPE_STORAGE_BUFFER:
+		return bRobustBufferAccess ? DescriptorBufferProperties.robustStorageBufferDescriptorSize
+			: DescriptorBufferProperties.storageBufferDescriptorSize;
+	case VK_DESCRIPTOR_TYPE_ACCELERATION_STRUCTURE_KHR:
+		return DescriptorBufferProperties.accelerationStructureDescriptorSize;
+	default: checkNoEntry();
+	}
+	return 0;
+}
+
+
+static inline uint32 GetInitialDescriptorCount(VkDescriptorType DescriptorType)
+{
+	switch (DescriptorType)
+	{
+	case VK_DESCRIPTOR_TYPE_SAMPLER:                    return GVulkanBindlessMaxSamplerDescriptorCount;
+	case VK_DESCRIPTOR_TYPE_SAMPLED_IMAGE:              return GVulkanBindlessMaxSampledImageDescriptorCount;
+	case VK_DESCRIPTOR_TYPE_STORAGE_IMAGE:              return GVulkanBindlessMaxStorageImageDescriptorCount;
+	case VK_DESCRIPTOR_TYPE_UNIFORM_TEXEL_BUFFER:       return GVulkanBindlessMaxUniformTexelBufferDescriptorCount;
+	case VK_DESCRIPTOR_TYPE_STORAGE_TEXEL_BUFFER:       return GVulkanBindlessMaxStorageTexelBufferDescriptorCount;
+	case VK_DESCRIPTOR_TYPE_UNIFORM_BUFFER:             return GVulkanBindlessMaxUniformBufferDescriptorCount;
+	case VK_DESCRIPTOR_TYPE_STORAGE_BUFFER:             return GVulkanBindlessMaxStorageBufferDescriptorCount;
+	case VK_DESCRIPTOR_TYPE_ACCELERATION_STRUCTURE_KHR: return GVulkanBindlessMaxAccelerationStructureDescriptorCount;
+	default: checkNoEntry();
+	}
+	return 0;
+}
+
+
+static inline VkMemoryPropertyFlags GetDescriptorBufferMemoryType(FVulkanDevice& Device)
+{
+	if (Device.HasUnifiedMemory() || (FVulkanPlatform::SupportsDeviceLocalHostVisibleWithNoPenalty(Device.GetVendorId()) &&
+		Device.GetDeviceMemoryManager().SupportsMemoryType(VK_MEMORY_PROPERTY_HOST_VISIBLE_BIT | VK_MEMORY_PROPERTY_HOST_COHERENT_BIT | VK_MEMORY_PROPERTY_DEVICE_LOCAL_BIT)))
+	{
+		return VK_MEMORY_PROPERTY_DEVICE_LOCAL_BIT | VK_MEMORY_PROPERTY_HOST_VISIBLE_BIT | VK_MEMORY_PROPERTY_HOST_COHERENT_BIT;
+	}
+
+	return VK_MEMORY_PROPERTY_HOST_VISIBLE_BIT | VK_MEMORY_PROPERTY_HOST_COHERENT_BIT;
+}
+
+
+
+// Check all the requirements to be running in Bindless using Descriptor Buffers
+bool FVulkanBindlessDescriptorManager::VerifySupport(FVulkanDevice& Device)
+{
+	const ERHIBindlessConfiguration BindlessConfig = RHIGetRuntimeBindlessConfiguration(GMaxRHIShaderPlatform);
+
+	if (BindlessConfig == ERHIBindlessConfiguration::Disabled)
+	{
+		return false;
+	}
+
+	if (BindlessConfig > ERHIBindlessConfiguration::RayTracing)
+	{
+		const VkPhysicalDeviceProperties& GpuProps = Device.GetDeviceProperties();
+		const FOptionalVulkanDeviceExtensions& OptionalDeviceExtensions = Device.GetOptionalExtensions();
+		const VkPhysicalDeviceDescriptorBufferPropertiesEXT& DescriptorBufferProperties = Device.GetOptionalExtensionProperties().DescriptorBufferProps;
+
+		const bool bMeetsExtensionsRequirements =
+			OptionalDeviceExtensions.HasEXTDescriptorIndexing &&
+			OptionalDeviceExtensions.HasBufferDeviceAddress &&
+			OptionalDeviceExtensions.HasEXTDescriptorBuffer;
+
+		if (bMeetsExtensionsRequirements)
+		{
+			const bool bMeetsPropertiesRequirements =
+				(GpuProps.limits.maxBoundDescriptorSets >= VulkanBindless::MaxNumSets) &&
+				(DescriptorBufferProperties.maxDescriptorBufferBindings >= VulkanBindless::MaxNumSets) &&
+				(DescriptorBufferProperties.maxResourceDescriptorBufferBindings >= VulkanBindless::NumBindlessSets) &&
+				(DescriptorBufferProperties.maxSamplerDescriptorBufferBindings >= 1) &&
+				Device.GetDeviceMemoryManager().SupportsMemoryType(GetDescriptorBufferMemoryType(Device));
+
+			if (bMeetsPropertiesRequirements)
+			{
+				extern TAutoConsoleVariable<int32> GDynamicGlobalUBs;
+				if (GDynamicGlobalUBs->GetInt() != 0)
+				{
+					UE_LOG(LogRHI, Warning, TEXT("Dynamic Uniform Buffers are enabled, but they will not be used with Vulkan bindless."));
+				}
+
+				extern int32 GVulkanEnableDefrag;
+				if (GVulkanEnableDefrag != 0)  // :todo-jn: to be turned back on with new defragger
+				{
+					UE_LOG(LogRHI, Warning, TEXT("Memory defrag is enabled, but it will not be used with Vulkan bindless."));
+					GVulkanEnableDefrag = 0;
+				}
+
+				return true;
+			}
+			else
+			{
+				UE_LOG(LogRHI, Warning, TEXT("Bindless descriptor were requested but NOT enabled because of insufficient property support."));
+			}
+		}
+		else
+		{
+			UE_LOG(LogRHI, Warning, TEXT("Bindless descriptor were requested but NOT enabled because of missing extension support."));
+		}
+	}
+	else
+	{
+		UE_LOG(LogRHI, Warning, TEXT("Bindless in Vulkan must currently be fully enabled (all samplers and resources) or fully disabled."));
+	}
+
+	return false;
+}
+
+
+FVulkanBindlessDescriptorManager::FVulkanBindlessDescriptorManager(FVulkanDevice& InDevice)
+	: Device(InDevice)
+	, bIsSupported(VerifySupport(InDevice))
+{
+	FMemory::Memzero(BufferBindingInfo);
+	for (uint32 Index = 0; Index < VulkanBindless::MaxNumSets; Index++)
+	{
+		BufferIndices[Index] = Index;
+	}
+}
+
+FVulkanBindlessDescriptorManager::~FVulkanBindlessDescriptorManager()
+{
+	checkf(BindlessPipelineLayout == VK_NULL_HANDLE, TEXT("DeInit() was not called on FVulkanBindlessDescriptorManager!"));
+}
+
+void FVulkanBindlessDescriptorManager::Deinit()
+{
+	const VkDevice DeviceHandle = Device.GetInstanceHandle();
+
+	if (bIsSupported)
+	{
+		VulkanRHI::vkDestroyPipelineLayout(DeviceHandle, BindlessPipelineLayout, VULKAN_CPU_ALLOCATOR);
+		BindlessPipelineLayout = VK_NULL_HANDLE;
+
+		auto DestroyBindlessState = [DeviceHandle](BindlessSetState& State)
+			{
+				VulkanRHI::vkDestroyDescriptorSetLayout(DeviceHandle, State.DescriptorSetLayout, VULKAN_CPU_ALLOCATOR);
+				State.DescriptorSetLayout = VK_NULL_HANDLE;
+
+				VulkanRHI::vkDestroyBuffer(DeviceHandle, State.BufferHandle, VULKAN_CPU_ALLOCATOR);
+				State.BufferHandle = VK_NULL_HANDLE;
+
+				VulkanRHI::vkUnmapMemory(DeviceHandle, State.MemoryHandle);
+				VulkanRHI::vkFreeMemory(DeviceHandle, State.MemoryHandle, VULKAN_CPU_ALLOCATOR);
+				State.MemoryHandle = VK_NULL_HANDLE;
+			};
+
+		for (uint32 SetIndex = 0; SetIndex < VulkanBindless::NumBindlessSets; ++SetIndex)
+		{
+			BindlessSetState& State = BindlessSetStates[SetIndex];
+			if (!State.DescriptorTypes.IsEmpty())
+			{
+				DestroyBindlessState(State);
+			}
+		}
+
+		VulkanRHI::vkDestroyDescriptorSetLayout(DeviceHandle, SingleUseUBDescriptorSetLayout, VULKAN_CPU_ALLOCATOR);
+		SingleUseUBDescriptorSetLayout = VK_NULL_HANDLE;
+
+		VulkanRHI::vkDestroyDescriptorSetLayout(DeviceHandle, EmptyDescriptorSetLayout, VULKAN_CPU_ALLOCATOR);
+		EmptyDescriptorSetLayout = VK_NULL_HANDLE;
+
+		delete SingleUseUBAllocator;
+		SingleUseUBAllocator = nullptr;
+	}
+}
+
+void FVulkanBindlessDescriptorManager::Init()
+{
+	if (!bIsSupported)
+	{
+		return;
+	}
+
+	const VkDevice DeviceHandle = Device.GetInstanceHandle();
+	const VkPhysicalDeviceDescriptorBufferPropertiesEXT& DescriptorBufferProperties = Device.GetOptionalExtensionProperties().DescriptorBufferProps;
+
+	const VkBufferUsageFlags BufferUsageFlags = VK_BUFFER_USAGE_SHADER_DEVICE_ADDRESS_BIT | VK_BUFFER_USAGE_RESOURCE_DESCRIPTOR_BUFFER_BIT_EXT;
+	SingleUseUBAllocator = new VulkanRHI::FTempBlockAllocator(&Device, GVulkanBindlessBlockSize, DescriptorBufferProperties.descriptorBufferOffsetAlignment, BufferUsageFlags);
+
+	// Create the dummy layout for unsupported descriptor types
+	{
+		VkDescriptorSetLayoutCreateInfo EmptyDescriptorSetLayoutCreateInfo;
+		ZeroVulkanStruct(EmptyDescriptorSetLayoutCreateInfo, VK_STRUCTURE_TYPE_DESCRIPTOR_SET_LAYOUT_CREATE_INFO);
+		EmptyDescriptorSetLayoutCreateInfo.flags = VK_DESCRIPTOR_SET_LAYOUT_CREATE_DESCRIPTOR_BUFFER_BIT_EXT;
+		VERIFYVULKANRESULT(VulkanRHI::vkCreateDescriptorSetLayout(DeviceHandle, &EmptyDescriptorSetLayoutCreateInfo, VULKAN_CPU_ALLOCATOR, &EmptyDescriptorSetLayout));
+	}
+
+	{
+		auto InitBindlessSetState = [&BindlessSetStates = BindlessSetStates, &Device = Device](VkDescriptorType DescriptorType) {
+
+			const uint8 StateIndex = GetIndexForDescriptorType(DescriptorType);
+			BindlessSetState& OutState = BindlessSetStates[StateIndex];
+
+			OutState.DescriptorTypes.Add(DescriptorType);
+
+			OutState.DescriptorSize = FMath::Max(OutState.DescriptorSize, GetDescriptorTypeSize(Device, DescriptorType));
+			checkf((OutState.DescriptorSize > 0), TEXT("Descriptor Type [%s] returned an invalid descriptor size!"), VK_TYPE_TO_STRING(VkDescriptorType, DescriptorType));
+
+			OutState.MaxDescriptorCount += GetInitialDescriptorCount(DescriptorType);
+			checkf((OutState.MaxDescriptorCount > 0), TEXT("Descriptor Type [%s] returned an invalid descriptor count!"), VK_TYPE_TO_STRING(VkDescriptorType, DescriptorType));
+			};
+
+		// Go through all the supported descriptor types in bindless
+		InitBindlessSetState(VK_DESCRIPTOR_TYPE_SAMPLER);
+		InitBindlessSetState(VK_DESCRIPTOR_TYPE_SAMPLED_IMAGE);
+		InitBindlessSetState(VK_DESCRIPTOR_TYPE_STORAGE_IMAGE);
+		InitBindlessSetState(VK_DESCRIPTOR_TYPE_UNIFORM_TEXEL_BUFFER);
+		InitBindlessSetState(VK_DESCRIPTOR_TYPE_STORAGE_TEXEL_BUFFER);
+		InitBindlessSetState(VK_DESCRIPTOR_TYPE_STORAGE_BUFFER);
+		InitBindlessSetState(VK_DESCRIPTOR_TYPE_UNIFORM_BUFFER);
+
+		// Create a dummy buffer for acceleration structures when they aren't supported (or ray tracing is disabled)
+		if (Device.GetOptionalExtensions().HasRaytracingExtensions())
+		{
+			InitBindlessSetState(VK_DESCRIPTOR_TYPE_ACCELERATION_STRUCTURE_KHR);
+		}
+		else
+		{
+			const uint8 StateIndex = GetIndexForDescriptorType(VK_DESCRIPTOR_TYPE_ACCELERATION_STRUCTURE_KHR);
+			BindlessSetState& State = BindlessSetStates[StateIndex];
+			State.DescriptorTypes.Add(VK_DESCRIPTOR_TYPE_STORAGE_BUFFER);
+			State.DescriptorSize = GetDescriptorTypeSize(Device, VK_DESCRIPTOR_TYPE_STORAGE_BUFFER);
+			State.MaxDescriptorCount = 16;
+		}
+
+		// Fill the DescriptorSetLayout for a BindlessSetState
+		auto CreateDescriptorSetLayout = [&](const BindlessSetState& State) {
+
+			if (State.DescriptorTypes.Num() == 0)
+			{
+				return EmptyDescriptorSetLayout;
+			}
+			else
+			{
+				VkDescriptorSetLayoutBinding DescriptorSetLayoutBinding;
+				DescriptorSetLayoutBinding.binding = 0;
+				DescriptorSetLayoutBinding.descriptorCount = State.MaxDescriptorCount;  // todo-jn: resizable
+				DescriptorSetLayoutBinding.stageFlags = VK_SHADER_STAGE_ALL;
+				DescriptorSetLayoutBinding.pImmutableSamplers = nullptr; // todo-jn: ImmutableSamplers
+
+				// These flags are implied with descriptor_buffer: VK_DESCRIPTOR_BINDING_UPDATE_AFTER_BIND_BIT,VK_DESCRIPTOR_BINDING_PARTIALLY_BOUND_BIT,VK_DESCRIPTOR_BINDING_UPDATE_UNUSED_WHILE_PENDING_BIT
+				// :todo-jn: add support for VK_DESCRIPTOR_BINDING_VARIABLE_DESCRIPTOR_COUNT_BIT when drivers are fixed to allow for buffers to grow
+				const VkDescriptorBindingFlags DescriptorBindingFlags = 0; // VK_DESCRIPTOR_BINDING_VARIABLE_DESCRIPTOR_COUNT_BIT;
+
+				VkDescriptorSetLayoutBindingFlagsCreateInfo DescriptorSetLayoutBindingFlagsCreateInfo;
+				ZeroVulkanStruct(DescriptorSetLayoutBindingFlagsCreateInfo, VK_STRUCTURE_TYPE_DESCRIPTOR_SET_LAYOUT_BINDING_FLAGS_CREATE_INFO);
+				DescriptorSetLayoutBindingFlagsCreateInfo.bindingCount = 1;
+				DescriptorSetLayoutBindingFlagsCreateInfo.pBindingFlags = &DescriptorBindingFlags;
+
+				VkDescriptorSetLayoutCreateInfo DescriptorSetLayoutCreateInfo;
+				ZeroVulkanStruct(DescriptorSetLayoutCreateInfo, VK_STRUCTURE_TYPE_DESCRIPTOR_SET_LAYOUT_CREATE_INFO);
+				DescriptorSetLayoutCreateInfo.pBindings = &DescriptorSetLayoutBinding;
+				DescriptorSetLayoutCreateInfo.bindingCount = 1;
+				DescriptorSetLayoutCreateInfo.flags = VK_DESCRIPTOR_SET_LAYOUT_CREATE_DESCRIPTOR_BUFFER_BIT_EXT;
+				DescriptorSetLayoutCreateInfo.pNext = &DescriptorSetLayoutBindingFlagsCreateInfo;
+
+				VkMutableDescriptorTypeCreateInfoEXT MutableDescriptorTypeCreateInfo;
+				VkMutableDescriptorTypeListEXT MutableDescriptorTypeList;
+				if (State.DescriptorTypes.Num() == 1)
+				{
+					DescriptorSetLayoutBinding.descriptorType = State.DescriptorTypes[0];
+				}
+				else
+				{
+					DescriptorSetLayoutBinding.descriptorType = VK_DESCRIPTOR_TYPE_MUTABLE_EXT;
+
+					MutableDescriptorTypeList.descriptorTypeCount = State.DescriptorTypes.Num();
+					MutableDescriptorTypeList.pDescriptorTypes = State.DescriptorTypes.GetData();
+
+					ZeroVulkanStruct(MutableDescriptorTypeCreateInfo, VK_STRUCTURE_TYPE_MUTABLE_DESCRIPTOR_TYPE_CREATE_INFO_EXT);
+					MutableDescriptorTypeCreateInfo.mutableDescriptorTypeListCount = 1;
+					MutableDescriptorTypeCreateInfo.pMutableDescriptorTypeLists = &MutableDescriptorTypeList;
+
+					MutableDescriptorTypeCreateInfo.pNext = DescriptorSetLayoutCreateInfo.pNext;
+					DescriptorSetLayoutCreateInfo.pNext = &MutableDescriptorTypeCreateInfo;
+				}
+
+				VkDescriptorSetLayout DescriptorSetLayout = VK_NULL_HANDLE;
+				VERIFYVULKANRESULT(VulkanRHI::vkCreateDescriptorSetLayout(DeviceHandle, &DescriptorSetLayoutCreateInfo, VULKAN_CPU_ALLOCATOR, &DescriptorSetLayout));
+				return DescriptorSetLayout;
+			}
+			};
+
+		// Create the descriptor buffer for a BindlessSetState
+		auto CreateDescriptorBuffer = [&](BindlessSetState& InOutState, VkDescriptorBufferBindingInfoEXT& OutBufferBindingInfo) {
+
+			// Skip unsupported descriptors
+			if (InOutState.DescriptorTypes.Num() == 0)
+			{
+				return 0u;
+			}
+
+			const bool IsSamplerSet = (InOutState.DescriptorTypes[0] == VK_DESCRIPTOR_TYPE_SAMPLER);
+			const VkBufferUsageFlags BufferUsageFlags =
+				VK_BUFFER_USAGE_SHADER_DEVICE_ADDRESS_BIT |
+				VK_BUFFER_USAGE_TRANSFER_SRC_BIT |
+				VK_BUFFER_USAGE_TRANSFER_DST_BIT |
+				(IsSamplerSet ? VK_BUFFER_USAGE_SAMPLER_DESCRIPTOR_BUFFER_BIT_EXT : VK_BUFFER_USAGE_RESOURCE_DESCRIPTOR_BUFFER_BIT_EXT);
+
+			const uint32 DescriptorBufferSize = InOutState.DescriptorSize * InOutState.MaxDescriptorCount;
+			InOutState.DebugDescriptors.SetNumZeroed(DescriptorBufferSize);
+
+			VkDeviceSize LayoutSizeInBytes = 0;
+			VulkanRHI::vkGetDescriptorSetLayoutSizeEXT(DeviceHandle, InOutState.DescriptorSetLayout, &LayoutSizeInBytes);
+			// Double check that the layout follows the rules for a single binding with an array of descriptors that are tightly packed
+			check(LayoutSizeInBytes == (InOutState.MaxDescriptorCount * InOutState.DescriptorSize));
+
+			if (IsSamplerSet)
+			{
+				checkf(DescriptorBufferSize < DescriptorBufferProperties.samplerDescriptorBufferAddressSpaceSize,
+					TEXT("Sampler descriptor buffer size [%u] exceeded maximum [%llu]."),
+					DescriptorBufferSize, DescriptorBufferProperties.samplerDescriptorBufferAddressSpaceSize);
+			}
+
+			// Create descriptor buffer
+			{
+				InOutState.BufferHandle = Device.CreateBuffer(DescriptorBufferSize, BufferUsageFlags);
+			}
+
+			// Allocate buffer memory, bind and map
+			{
+				VkMemoryRequirements BufferMemoryReqs;
+				FMemory::Memzero(BufferMemoryReqs);
+				VulkanRHI::vkGetBufferMemoryRequirements(DeviceHandle, InOutState.BufferHandle, &BufferMemoryReqs);
+				check(BufferMemoryReqs.size >= DescriptorBufferSize);
+
+				uint32 MemoryTypeIndex = 0;
+				VERIFYVULKANRESULT(Device.GetDeviceMemoryManager().GetMemoryTypeFromProperties(BufferMemoryReqs.memoryTypeBits, GetDescriptorBufferMemoryType(Device), &MemoryTypeIndex));
+
+				VkMemoryAllocateFlagsInfo FlagsInfo;
+				ZeroVulkanStruct(FlagsInfo, VK_STRUCTURE_TYPE_MEMORY_ALLOCATE_FLAGS_INFO);
+				FlagsInfo.flags = VK_MEMORY_ALLOCATE_DEVICE_ADDRESS_BIT;
+
+				VkMemoryAllocateInfo AllocateInfo;
+				ZeroVulkanStruct(AllocateInfo, VK_STRUCTURE_TYPE_MEMORY_ALLOCATE_INFO);
+				AllocateInfo.allocationSize = BufferMemoryReqs.size;
+				AllocateInfo.memoryTypeIndex = MemoryTypeIndex;
+				AllocateInfo.pNext = &FlagsInfo;
+
+				VERIFYVULKANRESULT(VulkanRHI::vkAllocateMemory(DeviceHandle, &AllocateInfo, VULKAN_CPU_ALLOCATOR, &InOutState.MemoryHandle));
+				VERIFYVULKANRESULT(VulkanRHI::vkBindBufferMemory(DeviceHandle, InOutState.BufferHandle, InOutState.MemoryHandle, 0));
+				VERIFYVULKANRESULT(VulkanRHI::vkMapMemory(DeviceHandle, InOutState.MemoryHandle, 0, VK_WHOLE_SIZE, 0, (void**)&InOutState.MappedPointer));
+				FMemory::Memzero(InOutState.MappedPointer, AllocateInfo.allocationSize);
+			}
+
+			// Setup the binding info
+			{
+				VkBufferDeviceAddressInfo AddressInfo;
+				ZeroVulkanStruct(AddressInfo, VK_STRUCTURE_TYPE_BUFFER_DEVICE_ADDRESS_INFO);
+				AddressInfo.buffer = InOutState.BufferHandle;
+
+				ZeroVulkanStruct(OutBufferBindingInfo, VK_STRUCTURE_TYPE_DESCRIPTOR_BUFFER_BINDING_INFO_EXT);
+				OutBufferBindingInfo.address = VulkanRHI::vkGetBufferDeviceAddressKHR(DeviceHandle, &AddressInfo);
+				OutBufferBindingInfo.usage = BufferUsageFlags;
+			}
+
+			return IsSamplerSet ? 0u : DescriptorBufferSize;
+			};
+
+		// Fill in one state for each descriptor type
+		uint32 TotalResourceDescriptorBufferSize = 0;
+		for (uint32 SetIndex = 0; SetIndex < VulkanBindless::NumBindlessSets; ++SetIndex)
+		{
+			BindlessSetState& State = BindlessSetStates[SetIndex];
+			State.DescriptorSetLayout = CreateDescriptorSetLayout(State);
+			TotalResourceDescriptorBufferSize += CreateDescriptorBuffer(State, BufferBindingInfo[SetIndex]);
+		}
+
+		// Fill in the state for single-use UB
+		// Uniform buffer descriptor set layout differ from the other resources, we reserve a fixed number of descriptors per stage for each draw/dispatch
+		// todo-jn: this could be compacted..
+		{
+			const uint32 NumTotalBindings = VulkanBindless::MaxUniformBuffersPerStage * ShaderStage::MaxNumStages;
+
+			TArray<VkDescriptorSetLayoutBinding> DescriptorSetLayoutBindings;
+			DescriptorSetLayoutBindings.SetNumZeroed(NumTotalBindings);
+			for (uint32 BindingIndex = 0; BindingIndex < NumTotalBindings; ++BindingIndex)
+			{
+				DescriptorSetLayoutBindings[BindingIndex].binding = BindingIndex;
+				DescriptorSetLayoutBindings[BindingIndex].descriptorType = VK_DESCRIPTOR_TYPE_UNIFORM_BUFFER;
+				DescriptorSetLayoutBindings[BindingIndex].descriptorCount = 1;
+				DescriptorSetLayoutBindings[BindingIndex].stageFlags = VK_SHADER_STAGE_ALL;
+			}
+
+			VkDescriptorSetLayoutCreateInfo DescriptorSetLayoutCreateInfo;
+			ZeroVulkanStruct(DescriptorSetLayoutCreateInfo, VK_STRUCTURE_TYPE_DESCRIPTOR_SET_LAYOUT_CREATE_INFO);
+			DescriptorSetLayoutCreateInfo.pBindings = DescriptorSetLayoutBindings.GetData();
+			DescriptorSetLayoutCreateInfo.bindingCount = NumTotalBindings;
+			DescriptorSetLayoutCreateInfo.flags = VK_DESCRIPTOR_SET_LAYOUT_CREATE_DESCRIPTOR_BUFFER_BIT_EXT;
+			DescriptorSetLayoutCreateInfo.pNext = nullptr;
+
+			checkSlow(SingleUseUBDescriptorSetLayout == VK_NULL_HANDLE);
+			VERIFYVULKANRESULT(VulkanRHI::vkCreateDescriptorSetLayout(DeviceHandle, &DescriptorSetLayoutCreateInfo, VULKAN_CPU_ALLOCATOR, &SingleUseUBDescriptorSetLayout));
+		}
+
+		checkf(TotalResourceDescriptorBufferSize < DescriptorBufferProperties.resourceDescriptorBufferAddressSpaceSize,
+			TEXT("Combined resource descriptor buffer size of [%u] exceeded maximum [%llu]."),
+			TotalResourceDescriptorBufferSize, DescriptorBufferProperties.resourceDescriptorBufferAddressSpaceSize);
+	}
+
+	// Now create the single pipeline layout used by everything
+	{
+		VkDescriptorSetLayout DescriptorSetLayouts[VulkanBindless::MaxNumSets];
+		for (int32 LayoutIndex = 0; LayoutIndex < VulkanBindless::NumBindlessSets; ++LayoutIndex)
+		{
+			const BindlessSetState& State = BindlessSetStates[LayoutIndex];
+			DescriptorSetLayouts[LayoutIndex] = State.DescriptorSetLayout;
+		}
+		DescriptorSetLayouts[VulkanBindless::BindlessSingleUseUniformBufferSet] = SingleUseUBDescriptorSetLayout;
+
+		VkPipelineLayoutCreateInfo PipelineLayoutCreateInfo;
+		ZeroVulkanStruct(PipelineLayoutCreateInfo, VK_STRUCTURE_TYPE_PIPELINE_LAYOUT_CREATE_INFO);
+		PipelineLayoutCreateInfo.setLayoutCount = VulkanBindless::MaxNumSets;
+		PipelineLayoutCreateInfo.pSetLayouts = DescriptorSetLayouts;
+		VERIFYVULKANRESULT(VulkanRHI::vkCreatePipelineLayout(DeviceHandle, &PipelineLayoutCreateInfo, VULKAN_CPU_ALLOCATOR, &BindlessPipelineLayout));
+		VULKAN_SET_DEBUG_NAME(Device, VK_OBJECT_TYPE_PIPELINE_LAYOUT, BindlessPipelineLayout, TEXT("BindlessPipelineLayout(SetCount=%d)"), VulkanBindless::MaxNumSets);
+	}
+}
+
+void FVulkanBindlessDescriptorManager::BindDescriptorBuffers(VkCommandBuffer CommandBuffer, VkPipelineStageFlags SupportedStages)
+{
+	checkf(bIsSupported, TEXT("Trying to BindDescriptorBuffers but bindless is not supported!"));
+
+	VulkanRHI::vkCmdBindDescriptorBuffersEXT(CommandBuffer, VulkanBindless::NumBindlessSets, BufferBindingInfo);
+
+	VkDeviceSize BufferOffsets[VulkanBindless::NumBindlessSets];
+	FMemory::Memzero(BufferOffsets);
+	if (SupportedStages & VK_PIPELINE_STAGE_FRAGMENT_SHADER_BIT)
+	{
+		VulkanRHI::vkCmdSetDescriptorBufferOffsetsEXT(CommandBuffer, VK_PIPELINE_BIND_POINT_GRAPHICS, BindlessPipelineLayout, 0, VulkanBindless::NumBindlessSets, BufferIndices, BufferOffsets);
+	}
+	if (SupportedStages & VK_PIPELINE_STAGE_COMPUTE_SHADER_BIT)
+	{
+		VulkanRHI::vkCmdSetDescriptorBufferOffsetsEXT(CommandBuffer, VK_PIPELINE_BIND_POINT_COMPUTE, BindlessPipelineLayout, 0, VulkanBindless::NumBindlessSets, BufferIndices, BufferOffsets);
+	}
+	if (SupportedStages & VK_PIPELINE_STAGE_RAY_TRACING_SHADER_BIT_KHR)
+	{
+		VulkanRHI::vkCmdSetDescriptorBufferOffsetsEXT(CommandBuffer, VK_PIPELINE_BIND_POINT_RAY_TRACING_KHR, BindlessPipelineLayout, 0, VulkanBindless::NumBindlessSets, BufferIndices, BufferOffsets);
+	}
+}
+
+void FVulkanBindlessDescriptorManager::RegisterUniformBuffers(FVulkanCommandListContext& Context, VkPipelineBindPoint BindPoint, const FUniformBufferDescriptorArrays& StageUBs)
+{
+	checkf(bIsSupported, TEXT("Trying to RegisterUniformBuffers but bindless is not supported!"));
+
+	// :todo-jn: Current uniform buffer layout is a bit wasteful with all the skipped bindings...
+	const uint32 UBDescriptorSize = GetDescriptorTypeSize(Device, VK_DESCRIPTOR_TYPE_UNIFORM_BUFFER);
+	const uint32 BlockDescriptorCount = VulkanBindless::MaxUniformBuffersPerStage * ShaderStage::MaxNumStages;
+	const uint32 BlockSize = BlockDescriptorCount * UBDescriptorSize;
+
+	VkDescriptorBufferBindingInfoEXT LocalBufferBindingInfo[VulkanBindless::MaxNumSets];
+	FMemory::Memcpy(LocalBufferBindingInfo, BufferBindingInfo, VulkanBindless::NumBindlessSets * sizeof(VkDescriptorBufferBindingInfoEXT));
+
+	VkDeviceSize BufferOffsets[VulkanBindless::MaxNumSets];
+	FMemory::Memzero(BufferOffsets);
+
+	uint8* MappedPointer = SingleUseUBAllocator->Alloc(BlockSize, Context,
+		LocalBufferBindingInfo[VulkanBindless::BindlessSingleUseUniformBufferSet], BufferOffsets[VulkanBindless::BindlessSingleUseUniformBufferSet]);
+
+	for (int32 StageIndex = 0; StageIndex < ShaderStage::MaxNumStages; ++StageIndex)
+	{
+		const TArray<VkDescriptorAddressInfoEXT>& DescriptorAddressInfos = StageUBs[StageIndex];
+
+		if (DescriptorAddressInfos.Num())
+		{
+			checkSlow(StageIndex < GetNumStagesForBindPoint(BindPoint));
+			check(DescriptorAddressInfos.Num() <= VulkanBindless::MaxUniformBuffersPerStage);
+			const int32 StageOffset = StageIndex * VulkanBindless::MaxUniformBuffersPerStage;
+
+			for (int32 DescriptorAddressInfoIndex = 0; DescriptorAddressInfoIndex < DescriptorAddressInfos.Num(); DescriptorAddressInfoIndex++)
+			{
+				const VkDescriptorAddressInfoEXT& DescriptorAddressInfo = DescriptorAddressInfos[DescriptorAddressInfoIndex];
+				checkSlow(DescriptorAddressInfo.sType != 0);  // make sure it was filled
+				checkSlow((DescriptorAddressInfo.range % 16) == 0);  // :todo-jn: make sure we don't trip on driver bug, to be removed on next release
+
+				const int32 BindingIndex = StageOffset + DescriptorAddressInfoIndex;
+				VkDeviceSize BindingByteOffset = 0;
+#if UE_BUILD_DEBUG
+				VulkanRHI::vkGetDescriptorSetLayoutBindingOffsetEXT(Device.GetInstanceHandle(), SingleUseUBDescriptorSetLayout, BindingIndex, &BindingByteOffset);
+				check(BindingByteOffset == BindingIndex * UBDescriptorSize);
+#else
+				BindingByteOffset = (BindingIndex * UBDescriptorSize);
+#endif
+
+				VkDescriptorGetInfoEXT Info;
+				ZeroVulkanStruct(Info, VK_STRUCTURE_TYPE_DESCRIPTOR_GET_INFO_EXT);
+				Info.type = VK_DESCRIPTOR_TYPE_UNIFORM_BUFFER;
+				Info.data.pUniformBuffer = &DescriptorAddressInfo;
+				VulkanRHI::vkGetDescriptorEXT(Device.GetInstanceHandle(), &Info, UBDescriptorSize, &MappedPointer[BindingByteOffset]);
+			}
+		}
+	}
+
+	// todo-jn: cache these states and only repeat when necessary
+	const VkCommandBuffer CommandBufferHandle = Context.GetCommandBuffer().GetHandle();
+	VulkanRHI::vkCmdBindDescriptorBuffersEXT(CommandBufferHandle, VulkanBindless::MaxNumSets, LocalBufferBindingInfo);
+	VulkanRHI::vkCmdSetDescriptorBufferOffsetsEXT(CommandBufferHandle, BindPoint, BindlessPipelineLayout, 0u, VulkanBindless::MaxNumSets, BufferIndices, BufferOffsets);
+}
+
+void FVulkanBindlessDescriptorManager::UpdateStatsForHandle(VkDescriptorType DescriptorType)
+{
+	const uint8 SetIndex = GetIndexForDescriptorType(DescriptorType);
+	const BindlessSetState& State = BindlessSetStates[SetIndex];
+
+	switch (DescriptorType)
+	{
+	case VK_DESCRIPTOR_TYPE_SAMPLER:                   SET_DWORD_STAT(STAT_VulkanBindlessPeakSampler, State.PeakDescriptorCount); break;
+	case VK_DESCRIPTOR_TYPE_SAMPLED_IMAGE:             SET_DWORD_STAT(STAT_VulkanBindlessPeakSampledImage, State.PeakDescriptorCount); break;
+	case VK_DESCRIPTOR_TYPE_STORAGE_IMAGE:             SET_DWORD_STAT(STAT_VulkanBindlessPeakStorageImage, State.PeakDescriptorCount); break;
+	case VK_DESCRIPTOR_TYPE_UNIFORM_TEXEL_BUFFER:      SET_DWORD_STAT(STAT_VulkanBindlessPeakUniformTexelBuffer, State.PeakDescriptorCount); break;
+	case VK_DESCRIPTOR_TYPE_STORAGE_TEXEL_BUFFER:      SET_DWORD_STAT(STAT_VulkanBindlessPeakStorageTexelBuffer, State.PeakDescriptorCount); break;
+	case VK_DESCRIPTOR_TYPE_STORAGE_BUFFER:            SET_DWORD_STAT(STAT_VulkanBindlessPeakStorageBuffer, State.PeakDescriptorCount); break;
+	case VK_DESCRIPTOR_TYPE_UNIFORM_BUFFER:            SET_DWORD_STAT(STAT_VulkanBindlessPeakUniformBuffer, State.PeakDescriptorCount); break;
+	case VK_DESCRIPTOR_TYPE_ACCELERATION_STRUCTURE_KHR:SET_DWORD_STAT(STAT_VulkanBindlessPeakAccelerationStructure, State.PeakDescriptorCount); break;
+	default:
+		checkf(false, TEXT("Unknown descriptor type for stat update: %s"), VK_TYPE_TO_STRING(VkDescriptorType, DescriptorType));
+	}
+}
+
+FRHIDescriptorHandle FVulkanBindlessDescriptorManager::ReserveDescriptor(VkDescriptorType DescriptorType)
+{
+	if (bIsSupported)
+	{
+		const uint8 SetIndex = GetIndexForDescriptorType(DescriptorType);
+		const ERHIDescriptorType RHIType = GetUEDescriptorTypeFromVkDescriptorType(DescriptorType);
+
+		BindlessSetState& State = BindlessSetStates[SetIndex];
+		const uint32 ResourceIndex = GetFreeResourceIndex(State);
+		return FRHIDescriptorHandle(RHIType, ResourceIndex);
+	}
+	return FRHIDescriptorHandle();
+}
+
+void FVulkanBindlessDescriptorManager::UpdateDescriptor(FRHIDescriptorHandle DescriptorHandle, VkDescriptorType DescriptorType, VkDescriptorDataEXT DescriptorData, bool bImmediateUpdate)
+{
+	checkf(DescriptorHandle.IsValid(), TEXT("Attemping to update invalid descriptor handle!"));
+
+	const uint8 SetIndex = GetIndexForDescriptorType(DescriptorHandle.GetType());
+
+	BindlessSetState& State = BindlessSetStates[SetIndex];
+	const uint32 ByteOffset = DescriptorHandle.GetIndex() * State.DescriptorSize;
+	checkSlow(State.DescriptorTypes.Contains(DescriptorType));
+
+	VkDescriptorGetInfoEXT Info;
+	ZeroVulkanStruct(Info, VK_STRUCTURE_TYPE_DESCRIPTOR_GET_INFO_EXT);
+	Info.type = DescriptorType;
+	Info.data = DescriptorData;
+	VulkanRHI::vkGetDescriptorEXT(Device.GetInstanceHandle(), &Info, State.DescriptorSize, &State.DebugDescriptors[ByteOffset]);
+
+	if (bImmediateUpdate)
+	{
+		FMemory::Memcpy(&State.MappedPointer[ByteOffset], &State.DebugDescriptors[ByteOffset], State.DescriptorSize);
+	}
+	else
+	{
+		check(!IsInRenderingThread() || FRHICommandListExecutor::GetImmediateCommandList().Bypass() || !IsRunningRHIInSeparateThread());
+		FVulkanCommandBuffer* CmdBuffer = Device.GetImmediateContext().GetActiveCmdBuffer();
+
+		// :todo-jn:  Hack to avoid barriers/copies in renderpasses
+		if (CmdBuffer->IsInsideRenderPass())
+		{
+			FMemory::Memcpy(&State.MappedPointer[ByteOffset], &State.DebugDescriptors[ByteOffset], State.DescriptorSize);
+		}
+		else
+		{
+			VulkanRHI::FStagingBuffer* StagingBuffer = Device.GetStagingManager().AcquireBuffer(State.DescriptorSize);
+			FMemory::Memcpy(StagingBuffer->GetMappedPointer(), &State.DebugDescriptors[ByteOffset], State.DescriptorSize);
+			{
+				VkMemoryBarrier2 MemoryBarrier;
+				ZeroVulkanStruct(MemoryBarrier, VK_STRUCTURE_TYPE_MEMORY_BARRIER_2);
+				MemoryBarrier.srcStageMask = VK_PIPELINE_STAGE_2_ALL_COMMANDS_BIT;
+				MemoryBarrier.srcAccessMask = VK_ACCESS_2_DESCRIPTOR_BUFFER_READ_BIT_EXT | VK_ACCESS_2_MEMORY_READ_BIT | VK_ACCESS_2_MEMORY_WRITE_BIT;
+				MemoryBarrier.dstStageMask = VK_PIPELINE_STAGE_2_TRANSFER_BIT;
+				MemoryBarrier.dstAccessMask = VK_ACCESS_2_TRANSFER_READ_BIT | VK_ACCESS_2_TRANSFER_WRITE_BIT;
+
+				VkDependencyInfo DependencyInfo;
+				ZeroVulkanStruct(DependencyInfo, VK_STRUCTURE_TYPE_DEPENDENCY_INFO);
+				DependencyInfo.memoryBarrierCount = 1;
+				DependencyInfo.pMemoryBarriers = &MemoryBarrier;
+				VulkanRHI::vkCmdPipelineBarrier2KHR(CmdBuffer->GetHandle(), &DependencyInfo);
+
+				VkBufferCopy Region = {};
+				Region.srcOffset = 0;
+				Region.dstOffset = ByteOffset;
+				Region.size = State.DescriptorSize;
+				VulkanRHI::vkCmdCopyBuffer(CmdBuffer->GetHandle(), StagingBuffer->GetHandle(), State.BufferHandle, 1, &Region);
+
+				MemoryBarrier.srcStageMask = VK_PIPELINE_STAGE_2_TRANSFER_BIT;
+				MemoryBarrier.srcAccessMask = VK_ACCESS_2_TRANSFER_READ_BIT | VK_ACCESS_2_TRANSFER_WRITE_BIT;
+				MemoryBarrier.dstStageMask = VK_PIPELINE_STAGE_2_ALL_COMMANDS_BIT;
+				MemoryBarrier.dstAccessMask = VK_ACCESS_2_DESCRIPTOR_BUFFER_READ_BIT_EXT;
+				VulkanRHI::vkCmdPipelineBarrier2KHR(CmdBuffer->GetHandle(), &DependencyInfo);
+			}
+			Device.GetStagingManager().ReleaseBuffer(&Device.GetImmediateContext(), StagingBuffer);
+		}
+	}
+
+	UpdateStatsForHandle(DescriptorType);
+}
+
+void FVulkanBindlessDescriptorManager::UpdateSampler(FRHIDescriptorHandle DescriptorHandle, VkSampler VulkanSampler)
+{
+	if (bIsSupported)
+	{
+		VkDescriptorDataEXT DescriptorData;
+		DescriptorData.pSampler = &VulkanSampler;
+		UpdateDescriptor(DescriptorHandle, VK_DESCRIPTOR_TYPE_SAMPLER, DescriptorData, true);
+	}
+}
+
+void FVulkanBindlessDescriptorManager::UpdateImage(FRHIDescriptorHandle DescriptorHandle, VkDescriptorType DescriptorType, VkImageView ImageView, bool bIsDepthStencil, bool bImmediateUpdate)
+{
+	if (bIsSupported)
+	{
+		check((DescriptorType == VK_DESCRIPTOR_TYPE_SAMPLED_IMAGE) || (DescriptorType == VK_DESCRIPTOR_TYPE_STORAGE_IMAGE));
+
+		VkDescriptorImageInfo DescriptorImageInfo;
+		DescriptorImageInfo.sampler = VK_NULL_HANDLE;
+		DescriptorImageInfo.imageView = ImageView;
+		DescriptorImageInfo.imageLayout = (DescriptorType == VK_DESCRIPTOR_TYPE_STORAGE_IMAGE) ? VK_IMAGE_LAYOUT_GENERAL :
+			(bIsDepthStencil ? VK_IMAGE_LAYOUT_READ_ONLY_OPTIMAL : VK_IMAGE_LAYOUT_SHADER_READ_ONLY_OPTIMAL);
+
+		VkDescriptorDataEXT DescriptorData;
+		DescriptorData.pSampledImage = &DescriptorImageInfo;  // same pointer for storage, it's a union
+		UpdateDescriptor(DescriptorHandle, DescriptorType, DescriptorData, bImmediateUpdate);
+	}
+}
+
+void FVulkanBindlessDescriptorManager::UpdateBuffer(FRHIDescriptorHandle DescriptorHandle, VkDescriptorType DescriptorType, VkBuffer VulkanBuffer, VkDeviceSize BufferOffset, VkDeviceSize BufferSize, bool bImmediateUpdate)
+{
+	if (bIsSupported)
+	{
+		check((DescriptorType == VK_DESCRIPTOR_TYPE_UNIFORM_BUFFER) || (DescriptorType == VK_DESCRIPTOR_TYPE_STORAGE_BUFFER));
+
+		VkBufferDeviceAddressInfo BufferInfo;
+		ZeroVulkanStruct(BufferInfo, VK_STRUCTURE_TYPE_BUFFER_DEVICE_ADDRESS_INFO);
+		BufferInfo.buffer = VulkanBuffer;
+		const VkDeviceAddress BufferAddress = VulkanRHI::vkGetBufferDeviceAddressKHR(Device.GetInstanceHandle(), &BufferInfo);
+
+		UpdateBuffer(DescriptorHandle, DescriptorType, BufferAddress + BufferOffset, BufferSize, bImmediateUpdate);
+	}
+}
+
+void FVulkanBindlessDescriptorManager::UpdateBuffer(FRHIDescriptorHandle DescriptorHandle, VkDescriptorType DescriptorType, VkDeviceAddress BufferAddress, VkDeviceSize BufferSize, bool bImmediateUpdate)
+{
+	if (bIsSupported)
+	{
+		check((DescriptorType == VK_DESCRIPTOR_TYPE_UNIFORM_BUFFER) || (DescriptorType == VK_DESCRIPTOR_TYPE_STORAGE_BUFFER));
+
+		VkDescriptorAddressInfoEXT AddressInfo;
+		ZeroVulkanStruct(AddressInfo, VK_STRUCTURE_TYPE_DESCRIPTOR_ADDRESS_INFO_EXT);
+		AddressInfo.address = BufferAddress;
+		AddressInfo.range = BufferSize;
+
+		VkDescriptorDataEXT DescriptorData;
+		DescriptorData.pStorageBuffer = &AddressInfo;  // same pointer for uniform, it's a union
+		UpdateDescriptor(DescriptorHandle, DescriptorType, DescriptorData, bImmediateUpdate);
+	}
+}
+
+void FVulkanBindlessDescriptorManager::UpdateTexelBuffer(FRHIDescriptorHandle DescriptorHandle, VkDescriptorType DescriptorType, const VkBufferViewCreateInfo& ViewInfo, bool bImmediateUpdate)
+{
+	if (bIsSupported)
+	{
+		check((DescriptorType == VK_DESCRIPTOR_TYPE_UNIFORM_TEXEL_BUFFER) || (DescriptorType == VK_DESCRIPTOR_TYPE_STORAGE_TEXEL_BUFFER));
+
+		// :todo-jn: start caching buffer addresses in resources to avoid the extra call
+		VkBufferDeviceAddressInfo BufferInfo;
+		ZeroVulkanStruct(BufferInfo, VK_STRUCTURE_TYPE_BUFFER_DEVICE_ADDRESS_INFO);
+		BufferInfo.buffer = ViewInfo.buffer;
+		const VkDeviceAddress BufferAddress = VulkanRHI::vkGetBufferDeviceAddressKHR(Device.GetInstanceHandle(), &BufferInfo);
+
+		VkDescriptorAddressInfoEXT AddressInfo;
+		ZeroVulkanStruct(AddressInfo, VK_STRUCTURE_TYPE_DESCRIPTOR_ADDRESS_INFO_EXT);
+		AddressInfo.address = BufferAddress + ViewInfo.offset;
+		AddressInfo.range = ViewInfo.range;
+		AddressInfo.format = ViewInfo.format;
+
+		VkDescriptorDataEXT DescriptorData;
+		DescriptorData.pUniformTexelBuffer = &AddressInfo;  // same pointer for storage, it's a union
+		UpdateDescriptor(DescriptorHandle, DescriptorType, DescriptorData, bImmediateUpdate);
+	}
+}
+
+void FVulkanBindlessDescriptorManager::UpdateAccelerationStructure(FRHIDescriptorHandle DescriptorHandle, VkAccelerationStructureKHR AccelerationStructure, bool bImmediateUpdate)
+{
+	if (bIsSupported)
+	{
+		// :todo-jn: start caching AccelerationStructure in resources to avoid the extra call
+		VkAccelerationStructureDeviceAddressInfoKHR AccelerationStructureDeviceAddressInfo;
+		ZeroVulkanStruct(AccelerationStructureDeviceAddressInfo, VK_STRUCTURE_TYPE_ACCELERATION_STRUCTURE_DEVICE_ADDRESS_INFO_KHR);
+		AccelerationStructureDeviceAddressInfo.accelerationStructure = AccelerationStructure;
+		const VkDeviceAddress BufferAddress = VulkanRHI::vkGetAccelerationStructureDeviceAddressKHR(Device.GetInstanceHandle(), &AccelerationStructureDeviceAddressInfo);
+
+		VkDescriptorDataEXT DescriptorData;
+		DescriptorData.accelerationStructure = BufferAddress;
+		UpdateDescriptor(DescriptorHandle, VK_DESCRIPTOR_TYPE_ACCELERATION_STRUCTURE_KHR, DescriptorData, bImmediateUpdate);
+	}
+}
+
+
+uint32 FVulkanBindlessDescriptorManager::GetFreeResourceIndex(FVulkanBindlessDescriptorManager::BindlessSetState& State)
+{
+	INC_DWORD_STAT(STAT_VulkanBindlessWritePerFrame);
+
+	{
+		FScopeLock ScopeLock(&State.FreeListCS);
+		if ((State.FreeListHead != MAX_uint32) && (State.PeakDescriptorCount >= State.MaxDescriptorCount)) // todo-jn: temp
+		{
+			const uint32 FreeIndex = State.FreeListHead;
+			const uint32 ByteOffset = State.FreeListHead * State.DescriptorSize;
+			uint32* NextSlot = (uint32*)(&State.DebugDescriptors[ByteOffset]);
+			State.FreeListHead = *NextSlot;
+			return FreeIndex;
+		}
+	}
+
+	INC_DWORD_STAT(STAT_VulkanBindlessPeakDescriptorCount);
+
+	const uint32 ResourceIndex = State.PeakDescriptorCount++;
+	checkf(ResourceIndex < State.MaxDescriptorCount, TEXT("You need to grow the resource array size for [%s]!"), VK_TYPE_TO_STRING(VkDescriptorType, State.DescriptorTypes[0]));
+	return ResourceIndex;
+}
+
+void FVulkanBindlessDescriptorManager::Unregister(FRHIDescriptorHandle DescriptorHandle)
+{
+	if (DescriptorHandle.IsValid())
+	{
+		checkf(bIsSupported, TEXT("Unregistering a valid handle but bindless is not supported!"));
+
+		const uint8 SetIndex = GetIndexForDescriptorType(DescriptorHandle.GetType());
+		BindlessSetState& State = BindlessSetStates[SetIndex];
+
+		FScopeLock ScopeLock(&State.FreeListCS);
+
+		const uint32 PreviousHead = State.FreeListHead;
+		State.FreeListHead = DescriptorHandle.GetIndex();
+		const uint32 ByteOffset = DescriptorHandle.GetIndex() * State.DescriptorSize;
+		uint32* NewSlot = (uint32*)(&State.DebugDescriptors[ByteOffset]);
+		FMemory::Memzero(NewSlot, State.DescriptorSize); // easier for debugging for now
+		*NewSlot = PreviousHead;
+
+		// Clear the descriptor
+		// todo-jn: invalidate the GPU side?
+	}
+}
+
+void FVulkanBindlessDescriptorManager::UpdateUBAllocator()
+{
+	if (bIsSupported && SingleUseUBAllocator)
+	{
+		SingleUseUBAllocator->UpdateBlocks();
+	}
+}
diff --git a/Engine/Source/Runtime/VulkanRHI/Private/VulkanBindlessDescriptorManager.h b/Engine/Source/Runtime/VulkanRHI/Private/VulkanBindlessDescriptorManager.h
new file mode 100644
--- /dev/null
+++ b/Engine/Source/Runtime/VulkanRHI/Private/VulkanBindlessDescriptorManager.h
@@ -0,0 +1,89 @@
+// Copyright Epic Games, Inc. All Rights Reserved.
+
+#pragma once
+
+#include "Containers/StaticArray.h"
+#include "RHIDefinitions.h"
+#include "VulkanCommon.h"
+#include "VulkanMemory.h"
+#include "VulkanThirdParty.h"
+
+class FVulkanDevice;
+
+// Manager for resource descriptors used in bindless rendering.
+class FVulkanBindlessDescriptorManager
+{
+public:
+	FVulkanBindlessDescriptorManager(FVulkanDevice& InDevice);
+	~FVulkanBindlessDescriptorManager();
+
+	typedef TStaticArray<TArray<VkDescriptorAddressInfoEXT>, ShaderStage::MaxNumStages> FUniformBufferDescriptorArrays;
+
+	void Init();
+	void Deinit();
+
+	static bool VerifySupport(FVulkanDevice& InDevice);
+	inline bool IsSupported()
+	{
+		return bIsSupported;
+	}
+
+	inline VkPipelineLayout GetPipelineLayout() const
+	{
+		return BindlessPipelineLayout;
+	}
+
+	void BindDescriptorBuffers(VkCommandBuffer CommandBuffer, VkPipelineStageFlags SupportedStages);
+
+	FRHIDescriptorHandle ReserveDescriptor(VkDescriptorType DescriptorType);
+
+	void UpdateSampler(FRHIDescriptorHandle DescriptorHandle, VkSampler VulkanSampler);
+	void UpdateImage(FRHIDescriptorHandle DescriptorHandle, VkDescriptorType DescriptorType, VkImageView VulkanImage, bool bIsDepthStencil, bool bImmediateUpdate = true);
+	void UpdateBuffer(FRHIDescriptorHandle DescriptorHandle, VkDescriptorType DescriptorType, VkBuffer VulkanBuffer, VkDeviceSize BufferOffset, VkDeviceSize BufferSize, bool bImmediateUpdate = true);
+	void UpdateBuffer(FRHIDescriptorHandle DescriptorHandle, VkDescriptorType DescriptorType, VkDeviceAddress BufferAddress, VkDeviceSize BufferSize, bool bImmediateUpdate = true);
+	void UpdateTexelBuffer(FRHIDescriptorHandle DescriptorHandle, VkDescriptorType DescriptorType, const VkBufferViewCreateInfo& ViewInfo, bool bImmediateUpdate = true);
+	void UpdateAccelerationStructure(FRHIDescriptorHandle DescriptorHandle, VkAccelerationStructureKHR AccelerationStructure, bool bImmediateUpdate = true);
+
+	void RegisterUniformBuffers(FVulkanCommandListContext& Context, VkPipelineBindPoint BindPoint, const FUniformBufferDescriptorArrays& StageUBs);
+
+	void Unregister(FRHIDescriptorHandle DescriptorHandle);
+
+	void UpdateUBAllocator();
+
+private:
+	FVulkanDevice& Device;
+	const bool bIsSupported;
+
+	VkDescriptorSetLayout EmptyDescriptorSetLayout = VK_NULL_HANDLE;
+
+	struct BindlessSetState
+	{
+		TArray<VkDescriptorType> DescriptorTypes;
+		uint32 MaxDescriptorCount = 0;
+		std::atomic<uint32> PeakDescriptorCount = 1;  // always keep a null descriptor in slot 0
+		VkDescriptorSetLayout DescriptorSetLayout = VK_NULL_HANDLE;
+
+		FCriticalSection FreeListCS;
+		uint32 FreeListHead = MAX_uint32;
+
+		uint32 DescriptorSize = 0;
+		VkBuffer BufferHandle = VK_NULL_HANDLE;
+		VkDeviceMemory MemoryHandle = VK_NULL_HANDLE;
+		uint8* MappedPointer = nullptr;
+
+		TArray<uint8> DebugDescriptors;
+	};
+	BindlessSetState BindlessSetStates[VulkanBindless::NumBindlessSets];
+
+	VkDescriptorSetLayout SingleUseUBDescriptorSetLayout = VK_NULL_HANDLE;
+	VulkanRHI::FTempBlockAllocator* SingleUseUBAllocator = nullptr;
+
+	VkDescriptorBufferBindingInfoEXT BufferBindingInfo[VulkanBindless::NumBindlessSets];
+	uint32_t BufferIndices[VulkanBindless::MaxNumSets];
+
+	VkPipelineLayout BindlessPipelineLayout = VK_NULL_HANDLE;
+
+	uint32 GetFreeResourceIndex(BindlessSetState& Desc);
+	void UpdateStatsForHandle(VkDescriptorType DescriptorType);
+	void UpdateDescriptor(FRHIDescriptorHandle DescriptorHandle, VkDescriptorType DescriptorType, VkDescriptorDataEXT DescriptorData, bool bImmediateUpdate);
+};
diff --git a/Engine/Source/Runtime/VulkanRHI/Private/VulkanCommandBuffer.cpp b/Engine/Source/Runtime/VulkanRHI/Private/VulkanCommandBuffer.cpp
--- a/Engine/Source/Runtime/VulkanRHI/Private/VulkanCommandBuffer.cpp
+++ b/Engine/Source/Runtime/VulkanRHI/Private/VulkanCommandBuffer.cpp
@@ -6,10 +6,12 @@
 
 #include "VulkanCommandBuffer.h"
 #include "VulkanContext.h"
+#include "VulkanFramebuffer.h"
 #include "VulkanRenderpass.h"
 #include "VulkanDescriptorSets.h"
 #include "VulkanMemory.h"
 #include "VulkanRayTracing.h"
+#include "VulkanBindlessDescriptorManager.h"
 
 
 #define CMD_BUFFER_TIME_TO_WAIT_BEFORE_DELETING		10
diff --git a/Engine/Source/Runtime/VulkanRHI/Private/VulkanCommands.cpp b/Engine/Source/Runtime/VulkanRHI/Private/VulkanCommands.cpp
--- a/Engine/Source/Runtime/VulkanRHI/Private/VulkanCommands.cpp
+++ b/Engine/Source/Runtime/VulkanRHI/Private/VulkanCommands.cpp
@@ -13,6 +13,8 @@
 #include "GlobalRenderResources.h"
 #include "RHIShaderParametersShared.h"
 #include "RHIUtilities.h"
+#include "VulkanFramebuffer.h"
+#include "VulkanRenderpass.h"
 #include "VulkanResourceCollection.h"
 #include "RHICoreShader.h"
 #include "RHIUniformBufferUtilities.h"
diff --git a/Engine/Source/Runtime/VulkanRHI/Private/VulkanContext.cpp b/Engine/Source/Runtime/VulkanRHI/Private/VulkanContext.cpp
--- a/Engine/Source/Runtime/VulkanRHI/Private/VulkanContext.cpp
+++ b/Engine/Source/Runtime/VulkanRHI/Private/VulkanContext.cpp
@@ -2,7 +2,9 @@
 
 #include "VulkanContext.h"
 #include "VulkanCommandBuffer.h"
+#include "VulkanFramebuffer.h"
 #include "VulkanMemory.h"
+#include "VulkanRenderpass.h"
 #include "VulkanResources.h"
 #include "VulkanRHIPrivate.h"
 #include "VulkanDevice.h"
@@ -247,6 +249,14 @@ FVulkanCommandListContext::FVulkanCommandListContext(FVulkanDevice& InDevice, FV
 	GlobalUniformBuffers.AddZeroed(FUniformBufferStaticSlotRegistry::Get().GetSlotCount());
 }
 
+void FVulkanCommandListContext::NotifyDeletedRenderTarget(VkImage Image)
+{
+	if (CurrentFramebuffer && CurrentFramebuffer->ContainsRenderTarget(Image))
+	{
+		CurrentFramebuffer = nullptr;
+	}
+}
+
 
 void FVulkanCommandListContext::ReleasePendingState()
 {
diff --git a/Engine/Source/Runtime/VulkanRHI/Private/VulkanContext.h b/Engine/Source/Runtime/VulkanRHI/Private/VulkanContext.h
--- a/Engine/Source/Runtime/VulkanRHI/Private/VulkanContext.h
+++ b/Engine/Source/Runtime/VulkanRHI/Private/VulkanContext.h
@@ -371,13 +371,7 @@ public:
 		return PendingComputeState;
 	}
 
-	inline void NotifyDeletedRenderTarget(VkImage Image)
-	{
-		if (CurrentFramebuffer && CurrentFramebuffer->ContainsRenderTarget(Image))
-		{
-			CurrentFramebuffer = nullptr;
-		}
-	}
+	void NotifyDeletedRenderTarget(VkImage Image);
 
 	inline FVulkanRenderPass* GetCurrentRenderPass()
 	{
