Reorganizing a number of Vulkan types to be in their own file(s).

diff --git a/Engine/Source/Runtime/VulkanRHI/Private/VulkanDescriptorSets.cpp b/Engine/Source/Runtime/VulkanRHI/Private/VulkanDescriptorSets.cpp
--- a/Engine/Source/Runtime/VulkanRHI/Private/VulkanDescriptorSets.cpp
+++ b/Engine/Source/Runtime/VulkanRHI/Private/VulkanDescriptorSets.cpp
@@ -1,944 +1,526 @@
 // Copyright Epic Games, Inc. All Rights Reserved.
 
-/*=============================================================================
-	VulkanDescriptorSets.cpp: Vulkan descriptor set RHI implementation.
-=============================================================================*/
-
 #include "VulkanDescriptorSets.h"
-#include "VulkanRHIPrivate.h"
-#include "VulkanContext.h"
-
 
-int32 GVulkanBindlessMaxSamplerDescriptorCount = 2048;
-static FAutoConsoleVariableRef CVarVulkanBindlessMaxSamplerDescriptorCount(
-	TEXT("r.Vulkan.Bindless.MaxSamplerDescriptorCount"),
-	GVulkanBindlessMaxSamplerDescriptorCount,
-	TEXT("Maximum bindless sampler descriptor count"),
-	ECVF_ReadOnly
+TAutoConsoleVariable<int32> GDynamicGlobalUBs(
+	TEXT("r.Vulkan.DynamicGlobalUBs"),
+	2,
+	TEXT("2 to treat ALL uniform buffers as dynamic [default]\n")\
+	TEXT("1 to treat global/packed uniform buffers as dynamic\n")\
+	TEXT("0 to treat them as regular"),
+	ECVF_ReadOnly | ECVF_RenderThreadSafe
 );
 
-int32 GVulkanBindlessMaxSampledImageDescriptorCount = 256 * 1024;
-static FAutoConsoleVariableRef CVarVulkanBindlessMaxSampledImageCount(
-	TEXT("r.Vulkan.Bindless.MaxResourceSampledImageCount"),
-	GVulkanBindlessMaxSampledImageDescriptorCount,
-	TEXT("Maximum bindless Sampled Image descriptor count"),
-	ECVF_ReadOnly
-);
+void FVulkanDescriptorSetsLayoutInfo::ProcessBindingsForStage(VkShaderStageFlagBits StageFlags, ShaderStage::EStage DescSetStage, const FVulkanShaderHeader& CodeHeader, FUniformBufferGatherInfo& OutUBGatherInfo) const
+{
+	OutUBGatherInfo.CodeHeaders[DescSetStage] = &CodeHeader;
+}
 
-int32 GVulkanBindlessMaxStorageImageDescriptorCount = 64 * 1024;
-static FAutoConsoleVariableRef CVarVulkanBindlessMaxStorageImageCount(
-	TEXT("r.Vulkan.Bindless.MaxResourceStorageImageCount"),
-	GVulkanBindlessMaxStorageImageDescriptorCount,
-	TEXT("Maximum bindless Storage Image descriptor count"),
-	ECVF_ReadOnly
-);
+template<bool bIsCompute>
+void FVulkanDescriptorSetsLayoutInfo::FinalizeBindings(const FVulkanDevice& Device, const FUniformBufferGatherInfo& UBGatherInfo, const TArrayView<FRHISamplerState*>& ImmutableSamplers, bool bUsesBindless)
+{
+	// We'll be reusing this struct
+	VkDescriptorSetLayoutBinding Binding;
+	FMemory::Memzero(Binding);
+	Binding.descriptorCount = 1;
 
-int32 GVulkanBindlessMaxUniformTexelBufferDescriptorCount = 64 * 1024;
-static FAutoConsoleVariableRef CVarVulkanBindlessMaxUniformTexelBufferCount(
-	TEXT("r.Vulkan.Bindless.MaxResourceUniformTexelBufferCount"),
-	GVulkanBindlessMaxUniformTexelBufferDescriptorCount,
-	TEXT("Maximum bindless Uniform Texel Buffer descriptor count"),
-	ECVF_ReadOnly
-);
+	const bool bConvertAllUBsToDynamic = !bUsesBindless && (GDynamicGlobalUBs.GetValueOnAnyThread() > 1);
+	const bool bConvertPackedUBsToDynamic = !bUsesBindless && (bConvertAllUBsToDynamic || (GDynamicGlobalUBs.GetValueOnAnyThread() == 1));
+	const uint32 MaxDescriptorSetUniformBuffersDynamic = Device.GetLimits().maxDescriptorSetUniformBuffersDynamic;
 
-int32 GVulkanBindlessMaxStorageTexelBufferDescriptorCount = 64 * 1024;
-static FAutoConsoleVariableRef CVarVulkanBindlessMaxStorageTexelBufferCount(
-	TEXT("r.Vulkan.Bindless.MaxResourceStorageTexelBufferCount"),
-	GVulkanBindlessMaxStorageTexelBufferDescriptorCount,
-	TEXT("Maximum bindless Storage Texel Buffer descriptor count"),
-	ECVF_ReadOnly
-);
+	int32 CurrentImmutableSampler = 0;
+	for (int32 Stage = 0; Stage < (bIsCompute ? ShaderStage::NumComputeStages : ShaderStage::NumGraphicsStages); ++Stage)
+	{
+		checkSlow(StageInfos[Stage].IsEmpty());
 
-int32 GVulkanBindlessMaxUniformBufferDescriptorCount = 32 * 1024;
-static FAutoConsoleVariableRef CVarVulkanBindlessMaxUniformBufferCount(
-	TEXT("r.Vulkan.Bindless.MaxResourceUniformBufferCount"),
-	GVulkanBindlessMaxUniformBufferDescriptorCount,
-	TEXT("Maximum bindless Uniform Buffer descriptor count"),
-	ECVF_ReadOnly
-);
+		if (const FVulkanShaderHeader* ShaderHeader = UBGatherInfo.CodeHeaders[Stage])
+		{
+			FStageInfo& StageInfo = StageInfos[Stage];
 
-int32 GVulkanBindlessMaxStorageBufferDescriptorCount = 64 * 1024;
-static FAutoConsoleVariableRef CVarVulkanBindlessMaxStorageBufferCount(
-	TEXT("r.Vulkan.Bindless.MaxResourceStorageBufferCount"),
-	GVulkanBindlessMaxStorageBufferDescriptorCount,
-	TEXT("Maximum bindless Storage Buffer descriptor count"),
-	ECVF_ReadOnly
-);
+			Binding.stageFlags = UEFrequencyToVKStageBit(bIsCompute ? SF_Compute : ShaderStage::GetFrequencyForGfxStage((ShaderStage::EStage)Stage));
 
-int32 GVulkanBindlessMaxAccelerationStructureDescriptorCount = 64 * 1024;
-static FAutoConsoleVariableRef CVarVulkanBindlessMaxAccelerationStructureCount(
-	TEXT("r.Vulkan.Bindless.MaxResourceAccelerationStructureCount"),
-	GVulkanBindlessMaxAccelerationStructureDescriptorCount,
-	TEXT("Maximum bindless Acceleration Structure descriptor count"),
-	ECVF_ReadOnly
-);
+			StageInfo.PackedGlobalsSize = ShaderHeader->PackedGlobalsSize;
+			StageInfo.NumBoundUniformBuffers = ShaderHeader->NumBoundUniformBuffers;
 
-int32 GVulkanBindlessBlockSize = 1024 * 1024;
-static FAutoConsoleVariableRef CVarVulkanBindlessBlockSize(
-	TEXT("r.Vulkan.Bindless.BlockSize"),
-	GVulkanBindlessBlockSize,
-	TEXT("Block size to use for single use ub. (default: 1MB)"),
-	ECVF_RenderThreadSafe
-);
+			for (int32 BindingIndex = 0; BindingIndex < ShaderHeader->Bindings.Num(); ++BindingIndex)
+			{
+				const VkDescriptorType DescriptorType = (VkDescriptorType)ShaderHeader->Bindings[BindingIndex].DescriptorType;
 
+				const bool bIsUniformBuffer = (DescriptorType == VK_DESCRIPTOR_TYPE_UNIFORM_BUFFER);
+				const bool bIsGlobalPackedConstants = bIsUniformBuffer && ShaderHeader->PackedGlobalsSize && (BindingIndex == 0);
 
-DECLARE_STATS_GROUP(TEXT("Vulkan Bindless"), STATGROUP_VulkanBindless, STATCAT_Advanced);
-
-DECLARE_DWORD_ACCUMULATOR_STAT_EXTERN(TEXT("Num Peak Descriptor Count"), STAT_VulkanBindlessPeakDescriptorCount, STATGROUP_VulkanBindless, );
-DECLARE_DWORD_ACCUMULATOR_STAT_EXTERN(TEXT("Peak Samplers"), STAT_VulkanBindlessPeakSampler, STATGROUP_VulkanBindless, );
-DECLARE_DWORD_ACCUMULATOR_STAT_EXTERN(TEXT("Peak Sampled Images"), STAT_VulkanBindlessPeakSampledImage, STATGROUP_VulkanBindless, );
-DECLARE_DWORD_ACCUMULATOR_STAT_EXTERN(TEXT("Peak Storage Images"), STAT_VulkanBindlessPeakStorageImage, STATGROUP_VulkanBindless, );
-DECLARE_DWORD_ACCUMULATOR_STAT_EXTERN(TEXT("Peak Uniform Buffers"), STAT_VulkanBindlessPeakUniformBuffer, STATGROUP_VulkanBindless, );
-DECLARE_DWORD_ACCUMULATOR_STAT_EXTERN(TEXT("Peak Storage Buffers"), STAT_VulkanBindlessPeakStorageBuffer, STATGROUP_VulkanBindless, );
-DECLARE_DWORD_ACCUMULATOR_STAT_EXTERN(TEXT("Peak Uniform Texel Buffers"), STAT_VulkanBindlessPeakUniformTexelBuffer, STATGROUP_VulkanBindless, );
-DECLARE_DWORD_ACCUMULATOR_STAT_EXTERN(TEXT("Peak Storage Texel Buffers"), STAT_VulkanBindlessPeakStorageTexelBuffer, STATGROUP_VulkanBindless, );
-DECLARE_DWORD_ACCUMULATOR_STAT_EXTERN(TEXT("Peak Acceleration Structures"), STAT_VulkanBindlessPeakAccelerationStructure, STATGROUP_VulkanBindless, );
-DECLARE_DWORD_COUNTER_STAT_EXTERN(TEXT("Write Per Frame"), STAT_VulkanBindlessWritePerFrame, STATGROUP_VulkanBindless, );
-
-DEFINE_STAT(STAT_VulkanBindlessPeakDescriptorCount);
-DEFINE_STAT(STAT_VulkanBindlessPeakSampler);
-DEFINE_STAT(STAT_VulkanBindlessPeakSampledImage);
-DEFINE_STAT(STAT_VulkanBindlessPeakStorageImage);
-DEFINE_STAT(STAT_VulkanBindlessPeakUniformBuffer);
-DEFINE_STAT(STAT_VulkanBindlessPeakStorageBuffer);
-DEFINE_STAT(STAT_VulkanBindlessPeakUniformTexelBuffer);
-DEFINE_STAT(STAT_VulkanBindlessPeakStorageTexelBuffer);
-DEFINE_STAT(STAT_VulkanBindlessPeakAccelerationStructure);
-DEFINE_STAT(STAT_VulkanBindlessWritePerFrame);
-
-
-static constexpr uint8 GetIndexForDescriptorType(VkDescriptorType DescriptorType)
-{
-	switch (DescriptorType)
-	{
-	case VK_DESCRIPTOR_TYPE_SAMPLER:                    return VulkanBindless::BindlessSamplerSet;
-	case VK_DESCRIPTOR_TYPE_SAMPLED_IMAGE:              return VulkanBindless::BindlessSampledImageSet;
-	case VK_DESCRIPTOR_TYPE_STORAGE_IMAGE:              return VulkanBindless::BindlessStorageImageSet;
-	case VK_DESCRIPTOR_TYPE_UNIFORM_TEXEL_BUFFER:       return VulkanBindless::BindlessUniformTexelBufferSet;
-	case VK_DESCRIPTOR_TYPE_STORAGE_TEXEL_BUFFER:       return VulkanBindless::BindlessStorageTexelBufferSet;
-	case VK_DESCRIPTOR_TYPE_STORAGE_BUFFER:             return VulkanBindless::BindlessStorageBufferSet;
-	case VK_DESCRIPTOR_TYPE_UNIFORM_BUFFER:             return VulkanBindless::BindlessUniformBufferSet;
-	case VK_DESCRIPTOR_TYPE_ACCELERATION_STRUCTURE_KHR: return VulkanBindless::BindlessAccelerationStructureSet;
-	default: checkNoEntry();
-	}
+				if (bIsGlobalPackedConstants)
+				{
+					const VkDescriptorType UBType = bConvertPackedUBsToDynamic ? VK_DESCRIPTOR_TYPE_UNIFORM_BUFFER_DYNAMIC : VK_DESCRIPTOR_TYPE_UNIFORM_BUFFER;
 
-	return VulkanBindless::MaxNumSets;
-}
+					const uint32 NewBindingIndex = StageInfo.Types.Add(UBType);
+					checkf(NewBindingIndex == 0, TEXT("Packed globals should always be the first binding!"));
 
-static constexpr uint8 GetIndexForDescriptorType(ERHIDescriptorType DescriptorType)
-{
-	switch (DescriptorType)
-	{
-	case ERHIDescriptorType::Sampler:               return VulkanBindless::BindlessSamplerSet;
-	case ERHIDescriptorType::TextureSRV:            return VulkanBindless::BindlessSampledImageSet;
-	case ERHIDescriptorType::TextureUAV:            return VulkanBindless::BindlessStorageImageSet;
-	case ERHIDescriptorType::TypedBufferSRV:        return VulkanBindless::BindlessUniformTexelBufferSet;
-	case ERHIDescriptorType::TypedBufferUAV:        return VulkanBindless::BindlessStorageTexelBufferSet;
-	case ERHIDescriptorType::BufferUAV:             return VulkanBindless::BindlessStorageBufferSet;
-	case ERHIDescriptorType::BufferSRV:             return VulkanBindless::BindlessUniformBufferSet;
-	case ERHIDescriptorType::AccelerationStructure: return VulkanBindless::BindlessAccelerationStructureSet;
-	default: checkNoEntry();
+					Binding.binding = NewBindingIndex;
+					Binding.descriptorType = UBType;
+					AddDescriptor(Stage, Binding);
+				}
+				else if (bIsUniformBuffer)
+				{
+					VkDescriptorType UBType = VK_DESCRIPTOR_TYPE_UNIFORM_BUFFER;
+					if (bConvertAllUBsToDynamic && LayoutTypes[VK_DESCRIPTOR_TYPE_UNIFORM_BUFFER_DYNAMIC] < MaxDescriptorSetUniformBuffersDynamic)
+					{
+						UBType = VK_DESCRIPTOR_TYPE_UNIFORM_BUFFER_DYNAMIC;
+					}
+
+					// Here we might mess up with the stageFlags, so reset them every loop
+					Binding.descriptorType = UBType;
+					const FVulkanShaderHeader::FUniformBufferInfo& UBInfo = ShaderHeader->UniformBufferInfos[BindingIndex];
+					const bool bUBHasConstantData = (BindingIndex < (int32)ShaderHeader->NumBoundUniformBuffers);
+					if (bUBHasConstantData)
+					{
+						const uint32 NewBindingIndex = StageInfo.Types.Add(UBType);
+						check(NewBindingIndex == BindingIndex);
+						Binding.binding = NewBindingIndex;
+						AddDescriptor(Stage, Binding);
+					}
+				}
+				else
+				{
+					const uint32 NewTypeIndex = StageInfo.Types.Add(DescriptorType);
+					check(NewTypeIndex == BindingIndex);
+					Binding.binding = BindingIndex;
+					Binding.descriptorType = DescriptorType;
+					AddDescriptor(Stage, Binding);
+				}
+			}
+		}
 	}
 
-	return VulkanBindless::MaxNumSets;
+	CompileTypesUsageID();
+	GenerateHash(ImmutableSamplers, bIsCompute ? VK_PIPELINE_BIND_POINT_COMPUTE : VK_PIPELINE_BIND_POINT_GRAPHICS);
 }
 
+template void FVulkanDescriptorSetsLayoutInfo::FinalizeBindings<true>(const FVulkanDevice& Device, const FUniformBufferGatherInfo& UBGatherInfo, const TArrayView<FRHISamplerState*>& ImmutableSamplers, bool bUsesBindless);
+template void FVulkanDescriptorSetsLayoutInfo::FinalizeBindings<false>(const FVulkanDevice& Device, const FUniformBufferGatherInfo& UBGatherInfo, const TArrayView<FRHISamplerState*>& ImmutableSamplers, bool bUsesBindless);
 
-static constexpr ERHIDescriptorType GetUEDescriptorTypeFromVkDescriptorType(VkDescriptorType DescriptorType)
-{
-	switch (DescriptorType)
-	{
-	case VK_DESCRIPTOR_TYPE_SAMPLER:                    return ERHIDescriptorType::Sampler;
-	case VK_DESCRIPTOR_TYPE_SAMPLED_IMAGE:              return ERHIDescriptorType::TextureSRV;
-	case VK_DESCRIPTOR_TYPE_STORAGE_IMAGE:              return ERHIDescriptorType::TextureUAV;
-	case VK_DESCRIPTOR_TYPE_UNIFORM_TEXEL_BUFFER:       return ERHIDescriptorType::TypedBufferSRV;
-	case VK_DESCRIPTOR_TYPE_STORAGE_TEXEL_BUFFER:       return ERHIDescriptorType::TypedBufferUAV;
-	case VK_DESCRIPTOR_TYPE_STORAGE_BUFFER:             return ERHIDescriptorType::BufferUAV;
-	case VK_DESCRIPTOR_TYPE_UNIFORM_BUFFER:             return ERHIDescriptorType::BufferSRV;
-	case VK_DESCRIPTOR_TYPE_ACCELERATION_STRUCTURE_KHR: return ERHIDescriptorType::AccelerationStructure;
-	default: checkNoEntry();
-	}
 
-	return ERHIDescriptorType::Invalid;
+// Increments a value and asserts on overflow.
+// FSetInfo uses narrow integer types for descriptor counts,
+// which may feasibly overflow one day (for example if we add bindless resources).
+template <typename T>
+static void IncrementChecked(T& Value)
+{
+	check(Value < TNumericLimits<T>::Max());
+	++Value;
 }
 
-
-extern TAutoConsoleVariable<int32> GCVarRobustBufferAccess;
-static inline uint32 GetDescriptorTypeSize(FVulkanDevice& Device, VkDescriptorType DescriptorType)
+void FVulkanDescriptorSetsLayoutInfo::AddDescriptor(int32 DescriptorSetIndex, const VkDescriptorSetLayoutBinding& Descriptor)
 {
-	const bool bRobustBufferAccess = (GCVarRobustBufferAccess.GetValueOnAnyThread() > 0);
-	const VkPhysicalDeviceDescriptorBufferPropertiesEXT& DescriptorBufferProperties = Device.GetOptionalExtensionProperties().DescriptorBufferProps;
+	// Increment type usage
+	if (LayoutTypes.Contains(Descriptor.descriptorType))
+	{
+		LayoutTypes[Descriptor.descriptorType]++;
+	}
+	else
+	{
+		LayoutTypes.Add(Descriptor.descriptorType, 1);
+	}
 
-	switch (DescriptorType)
+	if (DescriptorSetIndex >= SetLayouts.Num())
 	{
-	case VK_DESCRIPTOR_TYPE_SAMPLER:                    
-		return DescriptorBufferProperties.samplerDescriptorSize;
-	case VK_DESCRIPTOR_TYPE_SAMPLED_IMAGE:              
-		return DescriptorBufferProperties.sampledImageDescriptorSize;
-	case VK_DESCRIPTOR_TYPE_STORAGE_IMAGE:              
-		return DescriptorBufferProperties.storageImageDescriptorSize;
-	case VK_DESCRIPTOR_TYPE_UNIFORM_TEXEL_BUFFER:       
-		return bRobustBufferAccess	? DescriptorBufferProperties.robustUniformTexelBufferDescriptorSize 
-									: DescriptorBufferProperties.uniformTexelBufferDescriptorSize;
-	case VK_DESCRIPTOR_TYPE_STORAGE_TEXEL_BUFFER:       
-		return bRobustBufferAccess	? DescriptorBufferProperties.robustStorageTexelBufferDescriptorSize 
-									: DescriptorBufferProperties.storageTexelBufferDescriptorSize;
-	case VK_DESCRIPTOR_TYPE_UNIFORM_BUFFER:             
-		return bRobustBufferAccess	? DescriptorBufferProperties.robustUniformBufferDescriptorSize 
-									: DescriptorBufferProperties.uniformBufferDescriptorSize;
-	case VK_DESCRIPTOR_TYPE_STORAGE_BUFFER:             
-		return bRobustBufferAccess	? DescriptorBufferProperties.robustStorageBufferDescriptorSize 
-									: DescriptorBufferProperties.storageBufferDescriptorSize;
-	case VK_DESCRIPTOR_TYPE_ACCELERATION_STRUCTURE_KHR: 
-		return DescriptorBufferProperties.accelerationStructureDescriptorSize;
-	default: checkNoEntry();
+		SetLayouts.SetNum(DescriptorSetIndex + 1, EAllowShrinking::No);
 	}
-	return 0;
-}
 
+	FSetLayout& DescSetLayout = SetLayouts[DescriptorSetIndex];
 
-static inline uint32 GetInitialDescriptorCount(VkDescriptorType DescriptorType)
-{
-	switch (DescriptorType)
+	VkDescriptorSetLayoutBinding* Binding = new(DescSetLayout.LayoutBindings) VkDescriptorSetLayoutBinding;
+	*Binding = Descriptor;
+
+	const FStageInfo& SetInfo = StageInfos[DescriptorSetIndex];
+	check(SetInfo.Types[Descriptor.binding] == Descriptor.descriptorType);
+	switch (Descriptor.descriptorType)
 	{
-	case VK_DESCRIPTOR_TYPE_SAMPLER:                    return GVulkanBindlessMaxSamplerDescriptorCount;
-	case VK_DESCRIPTOR_TYPE_SAMPLED_IMAGE:              return GVulkanBindlessMaxSampledImageDescriptorCount;
-	case VK_DESCRIPTOR_TYPE_STORAGE_IMAGE:              return GVulkanBindlessMaxStorageImageDescriptorCount;
-	case VK_DESCRIPTOR_TYPE_UNIFORM_TEXEL_BUFFER:       return GVulkanBindlessMaxUniformTexelBufferDescriptorCount;
-	case VK_DESCRIPTOR_TYPE_STORAGE_TEXEL_BUFFER:       return GVulkanBindlessMaxStorageTexelBufferDescriptorCount;
-	case VK_DESCRIPTOR_TYPE_UNIFORM_BUFFER:             return GVulkanBindlessMaxUniformBufferDescriptorCount;
-	case VK_DESCRIPTOR_TYPE_STORAGE_BUFFER:             return GVulkanBindlessMaxStorageBufferDescriptorCount;
-	case VK_DESCRIPTOR_TYPE_ACCELERATION_STRUCTURE_KHR: return GVulkanBindlessMaxAccelerationStructureDescriptorCount;
-	default: checkNoEntry();
+	case VK_DESCRIPTOR_TYPE_SAMPLER:
+	case VK_DESCRIPTOR_TYPE_COMBINED_IMAGE_SAMPLER:
+	case VK_DESCRIPTOR_TYPE_SAMPLED_IMAGE:
+	case VK_DESCRIPTOR_TYPE_STORAGE_IMAGE:
+	case VK_DESCRIPTOR_TYPE_INPUT_ATTACHMENT:
+		IncrementChecked(StageInfos[DescriptorSetIndex].NumImageInfos);
+		break;
+	case VK_DESCRIPTOR_TYPE_UNIFORM_BUFFER:
+	case VK_DESCRIPTOR_TYPE_STORAGE_BUFFER:
+	case VK_DESCRIPTOR_TYPE_UNIFORM_BUFFER_DYNAMIC:
+		IncrementChecked(StageInfos[DescriptorSetIndex].NumBufferInfos);
+		break;
+	case VK_DESCRIPTOR_TYPE_ACCELERATION_STRUCTURE_KHR:
+		IncrementChecked(StageInfos[DescriptorSetIndex].NumAccelerationStructures);
+		break;
+	case VK_DESCRIPTOR_TYPE_STORAGE_TEXEL_BUFFER:
+	case VK_DESCRIPTOR_TYPE_UNIFORM_TEXEL_BUFFER:
+		break;
+	default:
+		checkf(0, TEXT("Unsupported descriptor type %d"), (int32)Descriptor.descriptorType);
+		break;
 	}
-	return 0;
 }
 
-
-static inline VkMemoryPropertyFlags GetDescriptorBufferMemoryType(FVulkanDevice& Device)
+void FVulkanDescriptorSetsLayoutInfo::GenerateHash(const TArrayView<FRHISamplerState*>& InImmutableSamplers, VkPipelineBindPoint InBindPoint)
 {
-	if (Device.HasUnifiedMemory() || (FVulkanPlatform::SupportsDeviceLocalHostVisibleWithNoPenalty(Device.GetVendorId()) &&
-		Device.GetDeviceMemoryManager().SupportsMemoryType(VK_MEMORY_PROPERTY_HOST_VISIBLE_BIT | VK_MEMORY_PROPERTY_HOST_COHERENT_BIT | VK_MEMORY_PROPERTY_DEVICE_LOCAL_BIT)))
+	const int32 LayoutCount = SetLayouts.Num();
+	Hash = FCrc::MemCrc32(&TypesUsageID, sizeof(uint32), LayoutCount);
+
+	for (int32 layoutIndex = 0; layoutIndex < LayoutCount; ++layoutIndex)
 	{
-		return VK_MEMORY_PROPERTY_DEVICE_LOCAL_BIT | VK_MEMORY_PROPERTY_HOST_VISIBLE_BIT | VK_MEMORY_PROPERTY_HOST_COHERENT_BIT;
+		SetLayouts[layoutIndex].GenerateHash();
+		Hash = FCrc::MemCrc32(&SetLayouts[layoutIndex].Hash, sizeof(uint32), Hash);
 	}
-	else
+
+	const uint32 NumStages = GetNumStagesForBindPoint(InBindPoint);
+	for (uint32 RemapingIndex = 0; RemapingIndex < NumStages; ++RemapingIndex)
 	{
-		return VK_MEMORY_PROPERTY_HOST_VISIBLE_BIT | VK_MEMORY_PROPERTY_HOST_COHERENT_BIT;
-	}
-}
+		const FStageInfo& StageInfo = StageInfos[RemapingIndex];
 
+		Hash = FCrc::TypeCrc32(StageInfo.PackedGlobalsSize, Hash);
+		Hash = FCrc::TypeCrc32(StageInfo.NumBoundUniformBuffers, Hash);
+		Hash = FCrc::TypeCrc32(StageInfo.NumImageInfos, Hash);
+		Hash = FCrc::TypeCrc32(StageInfo.NumBufferInfos, Hash);
+		Hash = FCrc::TypeCrc32(StageInfo.NumAccelerationStructures, Hash);
 
+		const TArray<VkDescriptorType>& Types = StageInfo.Types;
+		Hash = FCrc::MemCrc32(Types.GetData(), sizeof(VkDescriptorType) * Types.Num(), Hash);
+	}
 
-// Check all the requirements to be running in Bindless using Descriptor Buffers
-bool FVulkanBindlessDescriptorManager::VerifySupport(FVulkanDevice& Device)
-{
-	const ERHIBindlessConfiguration BindlessConfig = RHIGetRuntimeBindlessConfiguration(GMaxRHIShaderPlatform);
+	// It would be better to store this when the object is created, but it's not available at that time, so we'll do it here.
+	BindPoint = InBindPoint;
 
-	if (BindlessConfig == ERHIBindlessConfiguration::Disabled)
-	{
-		return false;
-	}
+	// Include the bind point in the hash, because we can have graphics and compute PSOs with the same descriptor info, and we don't want them to collide.
+	Hash = FCrc::MemCrc32(&BindPoint, sizeof(BindPoint), Hash);
+}
 
-	if (BindlessConfig > ERHIBindlessConfiguration::RayTracing)
-	{
-		const VkPhysicalDeviceProperties& GpuProps = Device.GetDeviceProperties();
-		const FOptionalVulkanDeviceExtensions& OptionalDeviceExtensions = Device.GetOptionalExtensions();
-		const VkPhysicalDeviceDescriptorBufferPropertiesEXT& DescriptorBufferProperties = Device.GetOptionalExtensionProperties().DescriptorBufferProps;
+static FCriticalSection GTypesUsageCS;
+void FVulkanDescriptorSetsLayoutInfo::CompileTypesUsageID()
+{
+	FScopeLock ScopeLock(&GTypesUsageCS);
 
-		const bool bMeetsExtensionsRequirements =
-			OptionalDeviceExtensions.HasEXTDescriptorIndexing &&
-			OptionalDeviceExtensions.HasBufferDeviceAddress &&
-			OptionalDeviceExtensions.HasEXTDescriptorBuffer;
+	static TMap<uint32, uint32> GTypesUsageHashMap;
+	static uint32 GUniqueID = 1;
 
-		if (bMeetsExtensionsRequirements)
+	LayoutTypes.KeySort([](VkDescriptorType A, VkDescriptorType B)
 		{
-			const bool bMeetsPropertiesRequirements =
-				(GpuProps.limits.maxBoundDescriptorSets >= VulkanBindless::MaxNumSets) &&
-				(DescriptorBufferProperties.maxDescriptorBufferBindings >= VulkanBindless::MaxNumSets) &&
-				(DescriptorBufferProperties.maxResourceDescriptorBufferBindings >= VulkanBindless::NumBindlessSets) &&
-				(DescriptorBufferProperties.maxSamplerDescriptorBufferBindings >= 1) &&
-				Device.GetDeviceMemoryManager().SupportsMemoryType(GetDescriptorBufferMemoryType(Device));
-
-			if (bMeetsPropertiesRequirements)
-			{
-				extern TAutoConsoleVariable<int32> GDynamicGlobalUBs;
-				if (GDynamicGlobalUBs->GetInt() != 0)
-				{
-					UE_LOG(LogRHI, Warning, TEXT("Dynamic Uniform Buffers are enabled, but they will not be used with Vulkan bindless."));
-				}
+			return static_cast<uint32>(A) < static_cast<uint32>(B);
+		});
 
-				extern int32 GVulkanEnableDefrag;
-				if (GVulkanEnableDefrag != 0)  // :todo-jn: to be turned back on with new defragger
-				{
-					UE_LOG(LogRHI, Warning, TEXT("Memory defrag is enabled, but it will not be used with Vulkan bindless."));
-					GVulkanEnableDefrag = 0;
-				}
+	uint32 TypesUsageHash = 0;
+	for (const auto& Elem : LayoutTypes)
+	{
+		TypesUsageHash = FCrc::MemCrc32(&Elem.Value, sizeof(uint32), TypesUsageHash);
+	}
 
-				return true;
-			}
-			else
-			{
-				UE_LOG(LogRHI, Warning, TEXT("Bindless descriptor were requested but NOT enabled because of insufficient property support."));
-			}
-		}
-		else
-		{
-			UE_LOG(LogRHI, Warning, TEXT("Bindless descriptor were requested but NOT enabled because of missing extension support."));
-		}
+	uint32* UniqueID = GTypesUsageHashMap.Find(TypesUsageHash);
+	if (UniqueID == nullptr)
+	{
+		TypesUsageID = GTypesUsageHashMap.Add(TypesUsageHash, GUniqueID++);
 	}
 	else
 	{
-		UE_LOG(LogRHI, Warning, TEXT("Bindless in Vulkan must currently be fully enabled (all samplers and resources) or fully disabled."));
+		TypesUsageID = *UniqueID;
 	}
-
-	return false;
 }
 
+// 
+// FVulkanDescriptorSetsLayout
 
-FVulkanBindlessDescriptorManager::FVulkanBindlessDescriptorManager(FVulkanDevice& InDevice)
-	: Device(InDevice)
-	, bIsSupported(VerifySupport(InDevice))
+FVulkanDescriptorSetsLayout::FVulkanDescriptorSetsLayout(FVulkanDevice* InDevice) :
+	Device(InDevice)
 {
-	FMemory::Memzero(BufferBindingInfo);
-	for (uint32 Index = 0; Index < VulkanBindless::MaxNumSets; Index++)
-	{
-		BufferIndices[Index] = Index;
-	}
 }
 
-FVulkanBindlessDescriptorManager::~FVulkanBindlessDescriptorManager()
+FVulkanDescriptorSetsLayout::~FVulkanDescriptorSetsLayout()
 {
-	checkf(BindlessPipelineLayout == VK_NULL_HANDLE, TEXT("DeInit() was not called on FVulkanBindlessDescriptorManager!"));
+	// Handles are owned by FVulkanPipelineStateCacheManager
+	LayoutHandles.Reset(0);
 }
 
-void FVulkanBindlessDescriptorManager::Deinit()
+void FVulkanDescriptorSetsLayout::Compile(FVulkanDescriptorSetLayoutMap& DSetLayoutMap)
 {
-	const VkDevice DeviceHandle = Device.GetInstanceHandle();
+	check(LayoutHandles.Num() == 0);
 
-	if (bIsSupported)
-	{
-		VulkanRHI::vkDestroyPipelineLayout(DeviceHandle, BindlessPipelineLayout, VULKAN_CPU_ALLOCATOR);
-		BindlessPipelineLayout = VK_NULL_HANDLE;
+	// Check if we obey limits
+	const VkPhysicalDeviceLimits& Limits = Device->GetLimits();
 
-		auto DestroyBindlessState = [DeviceHandle](BindlessSetState& State)
-		{
-			VulkanRHI::vkDestroyDescriptorSetLayout(DeviceHandle, State.DescriptorSetLayout, VULKAN_CPU_ALLOCATOR);
-			State.DescriptorSetLayout = VK_NULL_HANDLE;
+	// Check for maxDescriptorSetSamplers
+	check(LayoutTypes[VK_DESCRIPTOR_TYPE_SAMPLER]
+		+ LayoutTypes[VK_DESCRIPTOR_TYPE_COMBINED_IMAGE_SAMPLER]
+			<= Limits.maxDescriptorSetSamplers);
 
-			VulkanRHI::vkDestroyBuffer(DeviceHandle, State.BufferHandle, VULKAN_CPU_ALLOCATOR);
-			State.BufferHandle = VK_NULL_HANDLE;
+	// Check for maxDescriptorSetUniformBuffers
+	check(LayoutTypes[VK_DESCRIPTOR_TYPE_UNIFORM_BUFFER]
+		+ LayoutTypes[VK_DESCRIPTOR_TYPE_UNIFORM_BUFFER_DYNAMIC]
+			<= Limits.maxDescriptorSetUniformBuffers);
 
-			VulkanRHI::vkUnmapMemory(DeviceHandle, State.MemoryHandle);
-			VulkanRHI::vkFreeMemory(DeviceHandle, State.MemoryHandle, VULKAN_CPU_ALLOCATOR);
-			State.MemoryHandle = VK_NULL_HANDLE;
-		};
+	// Check for maxDescriptorSetUniformBuffersDynamic
+	check(Device->GetVendorId() == EGpuVendorId::Amd ||
+		LayoutTypes[VK_DESCRIPTOR_TYPE_UNIFORM_BUFFER_DYNAMIC]
+		<= Limits.maxDescriptorSetUniformBuffersDynamic);
 
-		for (uint32 SetIndex = 0; SetIndex < VulkanBindless::NumBindlessSets; ++SetIndex)
-		{
-			BindlessSetState& State = BindlessSetStates[SetIndex];
-			if (!State.DescriptorTypes.IsEmpty())
-			{
-				DestroyBindlessState(State);
-			}
-		}
-		
-		VulkanRHI::vkDestroyDescriptorSetLayout(DeviceHandle, SingleUseUBDescriptorSetLayout, VULKAN_CPU_ALLOCATOR);
-		SingleUseUBDescriptorSetLayout = VK_NULL_HANDLE;
+	// Check for maxDescriptorSetStorageBuffers
+	check(LayoutTypes[VK_DESCRIPTOR_TYPE_STORAGE_BUFFER]
+		+ LayoutTypes[VK_DESCRIPTOR_TYPE_STORAGE_BUFFER_DYNAMIC]
+			<= Limits.maxDescriptorSetStorageBuffers);
 
-		VulkanRHI::vkDestroyDescriptorSetLayout(DeviceHandle, EmptyDescriptorSetLayout, VULKAN_CPU_ALLOCATOR);
-		EmptyDescriptorSetLayout = VK_NULL_HANDLE;
+	// Check for maxDescriptorSetStorageBuffersDynamic
+	check(LayoutTypes[VK_DESCRIPTOR_TYPE_STORAGE_BUFFER_DYNAMIC]
+		<= Limits.maxDescriptorSetStorageBuffersDynamic);
 
-		delete SingleUseUBAllocator;
-		SingleUseUBAllocator = nullptr;
-	}
-}
+	// Check for maxDescriptorSetSampledImages
+	check(LayoutTypes[VK_DESCRIPTOR_TYPE_COMBINED_IMAGE_SAMPLER]
+		+ LayoutTypes[VK_DESCRIPTOR_TYPE_SAMPLED_IMAGE]
+			+ LayoutTypes[VK_DESCRIPTOR_TYPE_UNIFORM_TEXEL_BUFFER]
+				<= Limits.maxDescriptorSetSampledImages);
 
-void FVulkanBindlessDescriptorManager::Init()
-{
-	if (!bIsSupported)
+	// Check for maxDescriptorSetStorageImages
+	check(LayoutTypes[VK_DESCRIPTOR_TYPE_STORAGE_IMAGE]
+		+ LayoutTypes[VK_DESCRIPTOR_TYPE_STORAGE_TEXEL_BUFFER]
+			<= Limits.maxDescriptorSetStorageImages);
+
+	check(LayoutTypes[VK_DESCRIPTOR_TYPE_INPUT_ATTACHMENT] <= Limits.maxDescriptorSetInputAttachments);
+
+	if (GRHISupportsRayTracing)
 	{
-		return;
+		check(LayoutTypes[VK_DESCRIPTOR_TYPE_ACCELERATION_STRUCTURE_KHR] < Device->GetOptionalExtensionProperties().AccelerationStructureProps.maxDescriptorSetAccelerationStructures);
 	}
 
-	const VkDevice DeviceHandle = Device.GetInstanceHandle();
-	const VkPhysicalDeviceDescriptorBufferPropertiesEXT& DescriptorBufferProperties = Device.GetOptionalExtensionProperties().DescriptorBufferProps;
-
-	const VkBufferUsageFlags BufferUsageFlags = VK_BUFFER_USAGE_SHADER_DEVICE_ADDRESS_BIT | VK_BUFFER_USAGE_RESOURCE_DESCRIPTOR_BUFFER_BIT_EXT;
-	SingleUseUBAllocator = new VulkanRHI::FTempBlockAllocator(&Device, GVulkanBindlessBlockSize, DescriptorBufferProperties.descriptorBufferOffsetAlignment, BufferUsageFlags);
+	LayoutHandles.Empty(SetLayouts.Num());
 
-	// Create the dummy layout for unsupported descriptor types
+	if (UseVulkanDescriptorCache())
 	{
-		VkDescriptorSetLayoutCreateInfo EmptyDescriptorSetLayoutCreateInfo;
-		ZeroVulkanStruct(EmptyDescriptorSetLayoutCreateInfo, VK_STRUCTURE_TYPE_DESCRIPTOR_SET_LAYOUT_CREATE_INFO);
-		EmptyDescriptorSetLayoutCreateInfo.flags = VK_DESCRIPTOR_SET_LAYOUT_CREATE_DESCRIPTOR_BUFFER_BIT_EXT;
-		VERIFYVULKANRESULT(VulkanRHI::vkCreateDescriptorSetLayout(DeviceHandle, &EmptyDescriptorSetLayoutCreateInfo, VULKAN_CPU_ALLOCATOR, &EmptyDescriptorSetLayout));
+		LayoutHandleIds.Empty(SetLayouts.Num());
 	}
 
+	for (FSetLayout& Layout : SetLayouts)
 	{
-		auto InitBindlessSetState = [&BindlessSetStates=BindlessSetStates, &Device=Device](VkDescriptorType DescriptorType) {
+		VkDescriptorSetLayout* LayoutHandle = new(LayoutHandles) VkDescriptorSetLayout;
 
-			const uint8 StateIndex = GetIndexForDescriptorType(DescriptorType);
-			BindlessSetState& OutState = BindlessSetStates[StateIndex];
-
-			OutState.DescriptorTypes.Add(DescriptorType);
-
-			OutState.DescriptorSize = FMath::Max(OutState.DescriptorSize, GetDescriptorTypeSize(Device, DescriptorType));
-			checkf((OutState.DescriptorSize > 0), TEXT("Descriptor Type [%s] returned an invalid descriptor size!"), VK_TYPE_TO_STRING(VkDescriptorType, DescriptorType));
-
-			OutState.MaxDescriptorCount += GetInitialDescriptorCount(DescriptorType);
-			checkf((OutState.MaxDescriptorCount > 0), TEXT("Descriptor Type [%s] returned an invalid descriptor count!"), VK_TYPE_TO_STRING(VkDescriptorType, DescriptorType));
-		};
-
-		// Go through all the supported descriptor types in bindless
-		InitBindlessSetState(VK_DESCRIPTOR_TYPE_SAMPLER);
-		InitBindlessSetState(VK_DESCRIPTOR_TYPE_SAMPLED_IMAGE);
-		InitBindlessSetState(VK_DESCRIPTOR_TYPE_STORAGE_IMAGE);
-		InitBindlessSetState(VK_DESCRIPTOR_TYPE_UNIFORM_TEXEL_BUFFER);
-		InitBindlessSetState(VK_DESCRIPTOR_TYPE_STORAGE_TEXEL_BUFFER);
-		InitBindlessSetState(VK_DESCRIPTOR_TYPE_STORAGE_BUFFER);
-		InitBindlessSetState(VK_DESCRIPTOR_TYPE_UNIFORM_BUFFER);
-
-		// Create a dummy buffer for acceleration structures when they aren't supported (or ray tracing is disabled)
-		if (Device.GetOptionalExtensions().HasRaytracingExtensions())
-		{
-			InitBindlessSetState(VK_DESCRIPTOR_TYPE_ACCELERATION_STRUCTURE_KHR);
-		}
-		else
+		uint32* LayoutHandleId = nullptr;
+		if (UseVulkanDescriptorCache())
 		{
-			const uint8 StateIndex = GetIndexForDescriptorType(VK_DESCRIPTOR_TYPE_ACCELERATION_STRUCTURE_KHR);
-			BindlessSetState& State = BindlessSetStates[StateIndex];
-			State.DescriptorTypes.Add(VK_DESCRIPTOR_TYPE_STORAGE_BUFFER);
-			State.DescriptorSize = GetDescriptorTypeSize(Device, VK_DESCRIPTOR_TYPE_STORAGE_BUFFER);
-			State.MaxDescriptorCount = 16;
+			LayoutHandleId = new(LayoutHandleIds) uint32;
 		}
 
-		// Fill the DescriptorSetLayout for a BindlessSetState
-		auto CreateDescriptorSetLayout = [&](const BindlessSetState& State) {
-
-			if (State.DescriptorTypes.Num() == 0)
-			{
-				return EmptyDescriptorSetLayout;
-			}
-			else
-			{
-				VkDescriptorSetLayoutBinding DescriptorSetLayoutBinding;
-				DescriptorSetLayoutBinding.binding = 0;
-				DescriptorSetLayoutBinding.descriptorCount = State.MaxDescriptorCount;  // todo-jn: resizable
-				DescriptorSetLayoutBinding.stageFlags = VK_SHADER_STAGE_ALL;
-				DescriptorSetLayoutBinding.pImmutableSamplers = nullptr; // todo-jn: ImmutableSamplers
-
-				// These flags are implied with descriptor_buffer: VK_DESCRIPTOR_BINDING_UPDATE_AFTER_BIND_BIT,VK_DESCRIPTOR_BINDING_PARTIALLY_BOUND_BIT,VK_DESCRIPTOR_BINDING_UPDATE_UNUSED_WHILE_PENDING_BIT
-				// :todo-jn: add support for VK_DESCRIPTOR_BINDING_VARIABLE_DESCRIPTOR_COUNT_BIT when drivers are fixed to allow for buffers to grow
-				const VkDescriptorBindingFlags DescriptorBindingFlags = 0; // VK_DESCRIPTOR_BINDING_VARIABLE_DESCRIPTOR_COUNT_BIT;
-
-				VkDescriptorSetLayoutBindingFlagsCreateInfo DescriptorSetLayoutBindingFlagsCreateInfo;
-				ZeroVulkanStruct(DescriptorSetLayoutBindingFlagsCreateInfo, VK_STRUCTURE_TYPE_DESCRIPTOR_SET_LAYOUT_BINDING_FLAGS_CREATE_INFO);
-				DescriptorSetLayoutBindingFlagsCreateInfo.bindingCount = 1;
-				DescriptorSetLayoutBindingFlagsCreateInfo.pBindingFlags = &DescriptorBindingFlags;
-
-				VkDescriptorSetLayoutCreateInfo DescriptorSetLayoutCreateInfo;
-				ZeroVulkanStruct(DescriptorSetLayoutCreateInfo, VK_STRUCTURE_TYPE_DESCRIPTOR_SET_LAYOUT_CREATE_INFO);
-				DescriptorSetLayoutCreateInfo.pBindings = &DescriptorSetLayoutBinding;
-				DescriptorSetLayoutCreateInfo.bindingCount = 1;
-				DescriptorSetLayoutCreateInfo.flags = VK_DESCRIPTOR_SET_LAYOUT_CREATE_DESCRIPTOR_BUFFER_BIT_EXT;
-				DescriptorSetLayoutCreateInfo.pNext = &DescriptorSetLayoutBindingFlagsCreateInfo;
-
-				VkMutableDescriptorTypeCreateInfoEXT MutableDescriptorTypeCreateInfo;
-				VkMutableDescriptorTypeListEXT MutableDescriptorTypeList;
-				if (State.DescriptorTypes.Num() == 1)
-				{
-					DescriptorSetLayoutBinding.descriptorType = State.DescriptorTypes[0];
-				}
-				else
-				{
-					DescriptorSetLayoutBinding.descriptorType = VK_DESCRIPTOR_TYPE_MUTABLE_EXT;
-
-					MutableDescriptorTypeList.descriptorTypeCount = State.DescriptorTypes.Num();
-					MutableDescriptorTypeList.pDescriptorTypes = State.DescriptorTypes.GetData();
-
-					ZeroVulkanStruct(MutableDescriptorTypeCreateInfo, VK_STRUCTURE_TYPE_MUTABLE_DESCRIPTOR_TYPE_CREATE_INFO_EXT);
-					MutableDescriptorTypeCreateInfo.mutableDescriptorTypeListCount = 1;
-					MutableDescriptorTypeCreateInfo.pMutableDescriptorTypeLists = &MutableDescriptorTypeList;
-
-					MutableDescriptorTypeCreateInfo.pNext = DescriptorSetLayoutCreateInfo.pNext;
-					DescriptorSetLayoutCreateInfo.pNext = &MutableDescriptorTypeCreateInfo;
-				}
-
-				VkDescriptorSetLayout DescriptorSetLayout = VK_NULL_HANDLE;
-				VERIFYVULKANRESULT(VulkanRHI::vkCreateDescriptorSetLayout(DeviceHandle, &DescriptorSetLayoutCreateInfo, VULKAN_CPU_ALLOCATOR, &DescriptorSetLayout));
-				return DescriptorSetLayout;
-			}
-		};
-
-		// Create the descriptor buffer for a BindlessSetState
-		auto CreateDescriptorBuffer = [&](BindlessSetState& InOutState, VkDescriptorBufferBindingInfoEXT& OutBufferBindingInfo) {
-
-			// Skip unsupported descriptors
-			if (InOutState.DescriptorTypes.Num() == 0)
-			{
-				return 0u;
-			}
-
-			const bool IsSamplerSet = (InOutState.DescriptorTypes[0] == VK_DESCRIPTOR_TYPE_SAMPLER);
-			const VkBufferUsageFlags BufferUsageFlags =
-				VK_BUFFER_USAGE_SHADER_DEVICE_ADDRESS_BIT |
-				VK_BUFFER_USAGE_TRANSFER_SRC_BIT |
-				VK_BUFFER_USAGE_TRANSFER_DST_BIT |
-				(IsSamplerSet ? VK_BUFFER_USAGE_SAMPLER_DESCRIPTOR_BUFFER_BIT_EXT : VK_BUFFER_USAGE_RESOURCE_DESCRIPTOR_BUFFER_BIT_EXT);
-
-			const uint32 DescriptorBufferSize = InOutState.DescriptorSize * InOutState.MaxDescriptorCount;
-			InOutState.DebugDescriptors.SetNumZeroed(DescriptorBufferSize);
-
-			VkDeviceSize LayoutSizeInBytes = 0;
-			VulkanRHI::vkGetDescriptorSetLayoutSizeEXT(DeviceHandle, InOutState.DescriptorSetLayout, &LayoutSizeInBytes);
-			// Double check that the layout follows the rules for a single binding with an array of descriptors that are tightly packed
-			check(LayoutSizeInBytes == (InOutState.MaxDescriptorCount * InOutState.DescriptorSize));
-
-			if (IsSamplerSet)
-			{
-				checkf(DescriptorBufferSize < DescriptorBufferProperties.samplerDescriptorBufferAddressSpaceSize,
-					TEXT("Sampler descriptor buffer size [%u] exceeded maximum [%llu]."),
-					DescriptorBufferSize, DescriptorBufferProperties.samplerDescriptorBufferAddressSpaceSize);
-			}
-
-			// Create descriptor buffer
-			{
-				InOutState.BufferHandle = Device.CreateBuffer(DescriptorBufferSize, BufferUsageFlags);
-			}
-
-			// Allocate buffer memory, bind and map
+		if (FVulkanDescriptorSetLayoutEntry* Found = DSetLayoutMap.Find(Layout))
+		{
+			*LayoutHandle = Found->Handle;
+			if (LayoutHandleId)
 			{
-				VkMemoryRequirements BufferMemoryReqs;
-				FMemory::Memzero(BufferMemoryReqs);
-				VulkanRHI::vkGetBufferMemoryRequirements(DeviceHandle, InOutState.BufferHandle, &BufferMemoryReqs);
-				check(BufferMemoryReqs.size >= DescriptorBufferSize);
-
-				uint32 MemoryTypeIndex = 0;
-				VERIFYVULKANRESULT(Device.GetDeviceMemoryManager().GetMemoryTypeFromProperties(BufferMemoryReqs.memoryTypeBits, GetDescriptorBufferMemoryType(Device), &MemoryTypeIndex));
-
-				VkMemoryAllocateFlagsInfo FlagsInfo;
-				ZeroVulkanStruct(FlagsInfo, VK_STRUCTURE_TYPE_MEMORY_ALLOCATE_FLAGS_INFO);
-				FlagsInfo.flags = VK_MEMORY_ALLOCATE_DEVICE_ADDRESS_BIT;
-
-				VkMemoryAllocateInfo AllocateInfo;
-				ZeroVulkanStruct(AllocateInfo, VK_STRUCTURE_TYPE_MEMORY_ALLOCATE_INFO);
-				AllocateInfo.allocationSize = BufferMemoryReqs.size;
-				AllocateInfo.memoryTypeIndex = MemoryTypeIndex;
-				AllocateInfo.pNext = &FlagsInfo;
-
-				VERIFYVULKANRESULT(VulkanRHI::vkAllocateMemory(DeviceHandle, &AllocateInfo, VULKAN_CPU_ALLOCATOR, &InOutState.MemoryHandle));
-				VERIFYVULKANRESULT(VulkanRHI::vkBindBufferMemory(DeviceHandle, InOutState.BufferHandle, InOutState.MemoryHandle, 0));
-				VERIFYVULKANRESULT(VulkanRHI::vkMapMemory(DeviceHandle, InOutState.MemoryHandle, 0, VK_WHOLE_SIZE, 0, (void**)&InOutState.MappedPointer));
-				FMemory::Memzero(InOutState.MappedPointer, AllocateInfo.allocationSize);
+				*LayoutHandleId = Found->HandleId;
 			}
+			continue;
+		}
 
-			// Setup the binding info
-			{
-				VkBufferDeviceAddressInfo AddressInfo;
-				ZeroVulkanStruct(AddressInfo, VK_STRUCTURE_TYPE_BUFFER_DEVICE_ADDRESS_INFO);
-				AddressInfo.buffer = InOutState.BufferHandle;
+		VkDescriptorSetLayoutCreateInfo DescriptorLayoutInfo;
+		ZeroVulkanStruct(DescriptorLayoutInfo, VK_STRUCTURE_TYPE_DESCRIPTOR_SET_LAYOUT_CREATE_INFO);
+		DescriptorLayoutInfo.bindingCount = Layout.LayoutBindings.Num();
+		DescriptorLayoutInfo.pBindings = Layout.LayoutBindings.GetData();
 
-				ZeroVulkanStruct(OutBufferBindingInfo, VK_STRUCTURE_TYPE_DESCRIPTOR_BUFFER_BINDING_INFO_EXT);
-				OutBufferBindingInfo.address = VulkanRHI::vkGetBufferDeviceAddressKHR(DeviceHandle, &AddressInfo);
-				OutBufferBindingInfo.usage = BufferUsageFlags;
-			}
+		VERIFYVULKANRESULT(VulkanRHI::vkCreateDescriptorSetLayout(Device->GetInstanceHandle(), &DescriptorLayoutInfo, VULKAN_CPU_ALLOCATOR, LayoutHandle));
 
-			return IsSamplerSet ? 0u : DescriptorBufferSize;
-		};
-
-		// Fill in one state for each descriptor type
-		uint32 TotalResourceDescriptorBufferSize = 0;
-		for (uint32 SetIndex = 0; SetIndex < VulkanBindless::NumBindlessSets; ++SetIndex)
+		if (LayoutHandleId)
 		{
-			BindlessSetState& State = BindlessSetStates[SetIndex];
-			State.DescriptorSetLayout = CreateDescriptorSetLayout(State);
-			TotalResourceDescriptorBufferSize += CreateDescriptorBuffer(State, BufferBindingInfo[SetIndex]);
+			*LayoutHandleId = ++GVulkanDSetLayoutHandleIdCounter;
 		}
 
-		// Fill in the state for single-use UB
-		// Uniform buffer descriptor set layout differ from the other resources, we reserve a fixed number of descriptors per stage for each draw/dispatch
-		// todo-jn: this could be compacted..
-		{
-			const uint32 NumTotalBindings = VulkanBindless::MaxUniformBuffersPerStage * ShaderStage::MaxNumStages;
+		FVulkanDescriptorSetLayoutEntry DescriptorSetLayoutEntry;
+		DescriptorSetLayoutEntry.Handle = *LayoutHandle;
+		DescriptorSetLayoutEntry.HandleId = LayoutHandleId ? *LayoutHandleId : 0;
 
-			TArray<VkDescriptorSetLayoutBinding> DescriptorSetLayoutBindings;
-			DescriptorSetLayoutBindings.SetNumZeroed(NumTotalBindings);
-			for (uint32 BindingIndex = 0; BindingIndex < NumTotalBindings; ++BindingIndex)
-			{
-				DescriptorSetLayoutBindings[BindingIndex].binding = BindingIndex;
-				DescriptorSetLayoutBindings[BindingIndex].descriptorType = VK_DESCRIPTOR_TYPE_UNIFORM_BUFFER;
-				DescriptorSetLayoutBindings[BindingIndex].descriptorCount = 1;
-				DescriptorSetLayoutBindings[BindingIndex].stageFlags = VK_SHADER_STAGE_ALL;
-			}
-
-			VkDescriptorSetLayoutCreateInfo DescriptorSetLayoutCreateInfo;
-			ZeroVulkanStruct(DescriptorSetLayoutCreateInfo, VK_STRUCTURE_TYPE_DESCRIPTOR_SET_LAYOUT_CREATE_INFO);
-			DescriptorSetLayoutCreateInfo.pBindings = DescriptorSetLayoutBindings.GetData();
-			DescriptorSetLayoutCreateInfo.bindingCount = NumTotalBindings;
-			DescriptorSetLayoutCreateInfo.flags = VK_DESCRIPTOR_SET_LAYOUT_CREATE_DESCRIPTOR_BUFFER_BIT_EXT;
-			DescriptorSetLayoutCreateInfo.pNext = nullptr;
-
-			checkSlow(SingleUseUBDescriptorSetLayout == VK_NULL_HANDLE);
-			VERIFYVULKANRESULT(VulkanRHI::vkCreateDescriptorSetLayout(DeviceHandle, &DescriptorSetLayoutCreateInfo, VULKAN_CPU_ALLOCATOR, &SingleUseUBDescriptorSetLayout));
-		}
-
-		checkf(TotalResourceDescriptorBufferSize < DescriptorBufferProperties.resourceDescriptorBufferAddressSpaceSize,
-			TEXT("Combined resource descriptor buffer size of [%u] exceeded maximum [%llu]."),
-			TotalResourceDescriptorBufferSize, DescriptorBufferProperties.resourceDescriptorBufferAddressSpaceSize);
+		DSetLayoutMap.Add(Layout, DescriptorSetLayoutEntry);
 	}
 
-	// Now create the single pipeline layout used by everything
+	if (TypesUsageID == ~0)
 	{
-		VkDescriptorSetLayout DescriptorSetLayouts[VulkanBindless::MaxNumSets];
-		for (int32 LayoutIndex = 0; LayoutIndex < VulkanBindless::NumBindlessSets; ++LayoutIndex)
-		{
-			const BindlessSetState& State = BindlessSetStates[LayoutIndex];
-			DescriptorSetLayouts[LayoutIndex] = State.DescriptorSetLayout;
-		}
-		DescriptorSetLayouts[VulkanBindless::BindlessSingleUseUniformBufferSet] = SingleUseUBDescriptorSetLayout;
-
-		VkPipelineLayoutCreateInfo PipelineLayoutCreateInfo;
-		ZeroVulkanStruct(PipelineLayoutCreateInfo, VK_STRUCTURE_TYPE_PIPELINE_LAYOUT_CREATE_INFO);
-		PipelineLayoutCreateInfo.setLayoutCount = VulkanBindless::MaxNumSets;
-		PipelineLayoutCreateInfo.pSetLayouts = DescriptorSetLayouts;
-		VERIFYVULKANRESULT(VulkanRHI::vkCreatePipelineLayout(DeviceHandle, &PipelineLayoutCreateInfo, VULKAN_CPU_ALLOCATOR, &BindlessPipelineLayout));
-		VULKAN_SET_DEBUG_NAME(Device, VK_OBJECT_TYPE_PIPELINE_LAYOUT, BindlessPipelineLayout, TEXT("BindlessPipelineLayout(SetCount=%d)"), VulkanBindless::MaxNumSets);
+		CompileTypesUsageID();
 	}
-}
-
-void FVulkanBindlessDescriptorManager::BindDescriptorBuffers(VkCommandBuffer CommandBuffer, VkPipelineStageFlags SupportedStages)
-{
-	checkf(bIsSupported, TEXT("Trying to BindDescriptorBuffers but bindless is not supported!"));
-
-	VulkanRHI::vkCmdBindDescriptorBuffersEXT(CommandBuffer, VulkanBindless::NumBindlessSets, BufferBindingInfo);
 
-	VkDeviceSize BufferOffsets[VulkanBindless::NumBindlessSets];
-	FMemory::Memzero(BufferOffsets);
-	if (SupportedStages & VK_PIPELINE_STAGE_FRAGMENT_SHADER_BIT)
-	{
-		VulkanRHI::vkCmdSetDescriptorBufferOffsetsEXT(CommandBuffer, VK_PIPELINE_BIND_POINT_GRAPHICS, BindlessPipelineLayout, 0, VulkanBindless::NumBindlessSets, BufferIndices, BufferOffsets);
-	}
-	if (SupportedStages & VK_PIPELINE_STAGE_COMPUTE_SHADER_BIT)
-	{
-		VulkanRHI::vkCmdSetDescriptorBufferOffsetsEXT(CommandBuffer, VK_PIPELINE_BIND_POINT_COMPUTE, BindlessPipelineLayout, 0, VulkanBindless::NumBindlessSets, BufferIndices, BufferOffsets);
-	}
-	if (SupportedStages & VK_PIPELINE_STAGE_RAY_TRACING_SHADER_BIT_KHR)
-	{
-		VulkanRHI::vkCmdSetDescriptorBufferOffsetsEXT(CommandBuffer, VK_PIPELINE_BIND_POINT_RAY_TRACING_KHR, BindlessPipelineLayout, 0, VulkanBindless::NumBindlessSets, BufferIndices, BufferOffsets);
-	}
+	ZeroVulkanStruct(DescriptorSetAllocateInfo, VK_STRUCTURE_TYPE_DESCRIPTOR_SET_ALLOCATE_INFO);
+	DescriptorSetAllocateInfo.descriptorSetCount = LayoutHandles.Num();
+	DescriptorSetAllocateInfo.pSetLayouts = LayoutHandles.GetData();
 }
 
-void FVulkanBindlessDescriptorManager::RegisterUniformBuffers(FVulkanCommandListContext& Context, VkPipelineBindPoint BindPoint, const FUniformBufferDescriptorArrays& StageUBs)
-{
-	checkf(bIsSupported, TEXT("Trying to RegisterUniformBuffers but bindless is not supported!"));
+//
+// FVulkanDescriptorSetWriter
 
-	// :todo-jn: Current uniform buffer layout is a bit wasteful with all the skipped bindings...
-	const uint32 UBDescriptorSize = GetDescriptorTypeSize(Device, VK_DESCRIPTOR_TYPE_UNIFORM_BUFFER);
-	const uint32 BlockDescriptorCount = VulkanBindless::MaxUniformBuffersPerStage * ShaderStage::MaxNumStages;
-	const uint32 BlockSize = BlockDescriptorCount * UBDescriptorSize;
+uint32 FVulkanDescriptorSetWriter::SetupDescriptorWrites(
+	const TArray<VkDescriptorType>& Types, FVulkanHashableDescriptorInfo* InHashableDescriptorInfos,
+	VkWriteDescriptorSet* InWriteDescriptors, VkDescriptorImageInfo* InImageInfo, VkDescriptorBufferInfo* InBufferInfo, uint8* InBindingToDynamicOffsetMap,
+	VkWriteDescriptorSetAccelerationStructureKHR* InAccelerationStructuresWriteDescriptors,
+	VkAccelerationStructureKHR* InAccelerationStructures,
+	const FVulkanSamplerState& DefaultSampler, const FVulkanView::FTextureView& DefaultImageView)
+{
+	HashableDescriptorInfos = InHashableDescriptorInfos;
+	WriteDescriptors = InWriteDescriptors;
+	NumWrites = Types.Num();
 
-	VkDescriptorBufferBindingInfoEXT LocalBufferBindingInfo[VulkanBindless::MaxNumSets];
-	FMemory::Memcpy(LocalBufferBindingInfo, BufferBindingInfo, VulkanBindless::NumBindlessSets * sizeof(VkDescriptorBufferBindingInfoEXT));
+	BindingToDynamicOffsetMap = InBindingToDynamicOffsetMap;
 
-	VkDeviceSize BufferOffsets[VulkanBindless::MaxNumSets];
-	FMemory::Memzero(BufferOffsets);
+	InitWrittenMasks(NumWrites);
 
-	uint8* MappedPointer = SingleUseUBAllocator->Alloc(BlockSize, Context,
-		LocalBufferBindingInfo[VulkanBindless::BindlessSingleUseUniformBufferSet], BufferOffsets[VulkanBindless::BindlessSingleUseUniformBufferSet]);
+	uint32 DynamicOffsetIndex = 0;
 
-	for (int32 StageIndex = 0; StageIndex < ShaderStage::MaxNumStages; ++StageIndex)
+	for (int32 Index = 0; Index < Types.Num(); ++Index)
 	{
-		const TArray<VkDescriptorAddressInfoEXT>& DescriptorAddressInfos = StageUBs[StageIndex];
+		InWriteDescriptors->sType = VK_STRUCTURE_TYPE_WRITE_DESCRIPTOR_SET;
+		InWriteDescriptors->dstBinding = Index;
+		InWriteDescriptors->descriptorCount = 1;
+		InWriteDescriptors->descriptorType = Types[Index];
 
-		if (DescriptorAddressInfos.Num())
+		switch (Types[Index])
 		{
-			checkSlow(StageIndex < GetNumStagesForBindPoint(BindPoint));
-			check(DescriptorAddressInfos.Num() <= VulkanBindless::MaxUniformBuffersPerStage);
-			const int32 StageOffset = StageIndex * VulkanBindless::MaxUniformBuffersPerStage;
-
-			for (int32 DescriptorAddressInfoIndex = 0; DescriptorAddressInfoIndex < DescriptorAddressInfos.Num(); DescriptorAddressInfoIndex++)
+		case VK_DESCRIPTOR_TYPE_UNIFORM_BUFFER_DYNAMIC:
+			BindingToDynamicOffsetMap[Index] = DynamicOffsetIndex;
+			++DynamicOffsetIndex;
+			InWriteDescriptors->pBufferInfo = InBufferInfo++;
+			break;
+		case VK_DESCRIPTOR_TYPE_UNIFORM_BUFFER:
+		case VK_DESCRIPTOR_TYPE_STORAGE_BUFFER:
+			InWriteDescriptors->pBufferInfo = InBufferInfo++;
+			break;
+		case VK_DESCRIPTOR_TYPE_SAMPLER:
+			SetWrittenBase(Index); //samplers have a default setting, don't assert on those yet.
+		case VK_DESCRIPTOR_TYPE_COMBINED_IMAGE_SAMPLER:
+		case VK_DESCRIPTOR_TYPE_SAMPLED_IMAGE:
+		case VK_DESCRIPTOR_TYPE_STORAGE_IMAGE:
+		case VK_DESCRIPTOR_TYPE_INPUT_ATTACHMENT:
+			// Texture.Load() still requires a default sampler...
+			if (InHashableDescriptorInfos) // UseVulkanDescriptorCache()
 			{
-				const VkDescriptorAddressInfoEXT& DescriptorAddressInfo = DescriptorAddressInfos[DescriptorAddressInfoIndex];
-				checkSlow(DescriptorAddressInfo.sType != 0);  // make sure it was filled
-				checkSlow((DescriptorAddressInfo.range % 16) == 0);  // :todo-jn: make sure we don't trip on driver bug, to be removed on next release
-
-				const int32 BindingIndex = StageOffset + DescriptorAddressInfoIndex;
-				VkDeviceSize BindingByteOffset = 0;
-#if UE_BUILD_DEBUG
-				VulkanRHI::vkGetDescriptorSetLayoutBindingOffsetEXT(Device.GetInstanceHandle(), SingleUseUBDescriptorSetLayout, BindingIndex, &BindingByteOffset);
-				check(BindingByteOffset == BindingIndex * UBDescriptorSize);
-#else
-				BindingByteOffset = (BindingIndex * UBDescriptorSize);
-#endif
-
-				VkDescriptorGetInfoEXT Info;
-				ZeroVulkanStruct(Info, VK_STRUCTURE_TYPE_DESCRIPTOR_GET_INFO_EXT);
-				Info.type = VK_DESCRIPTOR_TYPE_UNIFORM_BUFFER;
-				Info.data.pUniformBuffer = &DescriptorAddressInfo;
-				VulkanRHI::vkGetDescriptorEXT(Device.GetInstanceHandle(), &Info, UBDescriptorSize, &MappedPointer[BindingByteOffset]);
+				InHashableDescriptorInfos[Index].Image.SamplerId = DefaultSampler.SamplerId;
+				InHashableDescriptorInfos[Index].Image.ImageViewId = DefaultImageView.ViewId;
+				InHashableDescriptorInfos[Index].Image.ImageLayout = static_cast<uint32>(VK_IMAGE_LAYOUT_GENERAL);
 			}
+			InImageInfo->sampler = DefaultSampler.Sampler;
+			InImageInfo->imageView = DefaultImageView.View;
+			InImageInfo->imageLayout = VK_IMAGE_LAYOUT_GENERAL;
+			InWriteDescriptors->pImageInfo = InImageInfo++;
+			break;
+		case VK_DESCRIPTOR_TYPE_STORAGE_TEXEL_BUFFER:
+		case VK_DESCRIPTOR_TYPE_UNIFORM_TEXEL_BUFFER:
+			break;
+		case VK_DESCRIPTOR_TYPE_ACCELERATION_STRUCTURE_KHR:
+			InAccelerationStructuresWriteDescriptors->sType = VK_STRUCTURE_TYPE_WRITE_DESCRIPTOR_SET_ACCELERATION_STRUCTURE_KHR;
+			InAccelerationStructuresWriteDescriptors->pNext = nullptr;
+			InAccelerationStructuresWriteDescriptors->accelerationStructureCount = 1;
+			InAccelerationStructuresWriteDescriptors->pAccelerationStructures = InAccelerationStructures++;
+			InWriteDescriptors->pNext = InAccelerationStructuresWriteDescriptors++;
+			break;
+		default:
+			checkf(0, TEXT("Unsupported descriptor type %d"), (int32)Types[Index]);
+			break;
 		}
+		++InWriteDescriptors;
 	}
 
-	// todo-jn: cache these states and only repeat when necessary
-	const VkCommandBuffer CommandBufferHandle = Context.GetCommandBuffer().GetHandle();
-	VulkanRHI::vkCmdBindDescriptorBuffersEXT(CommandBufferHandle, VulkanBindless::MaxNumSets, LocalBufferBindingInfo);
-	VulkanRHI::vkCmdSetDescriptorBufferOffsetsEXT(CommandBufferHandle, BindPoint, BindlessPipelineLayout, 0u, VulkanBindless::MaxNumSets, BufferIndices, BufferOffsets);
+	return DynamicOffsetIndex;
 }
 
-void FVulkanBindlessDescriptorManager::UpdateStatsForHandle(VkDescriptorType DescriptorType)
+void FVulkanDescriptorSetWriter::CheckAllWritten()
 {
-	const uint8 SetIndex = GetIndexForDescriptorType(DescriptorType);
-	const BindlessSetState& State = BindlessSetStates[SetIndex];
-
-	switch (DescriptorType)
-	{
-	case VK_DESCRIPTOR_TYPE_SAMPLER:                   SET_DWORD_STAT(STAT_VulkanBindlessPeakSampler, State.PeakDescriptorCount); break;
-	case VK_DESCRIPTOR_TYPE_SAMPLED_IMAGE:             SET_DWORD_STAT(STAT_VulkanBindlessPeakSampledImage, State.PeakDescriptorCount); break;
-	case VK_DESCRIPTOR_TYPE_STORAGE_IMAGE:             SET_DWORD_STAT(STAT_VulkanBindlessPeakStorageImage, State.PeakDescriptorCount); break;
-	case VK_DESCRIPTOR_TYPE_UNIFORM_TEXEL_BUFFER:      SET_DWORD_STAT(STAT_VulkanBindlessPeakUniformTexelBuffer, State.PeakDescriptorCount); break;
-	case VK_DESCRIPTOR_TYPE_STORAGE_TEXEL_BUFFER:      SET_DWORD_STAT(STAT_VulkanBindlessPeakStorageTexelBuffer, State.PeakDescriptorCount); break;
-	case VK_DESCRIPTOR_TYPE_STORAGE_BUFFER:            SET_DWORD_STAT(STAT_VulkanBindlessPeakStorageBuffer, State.PeakDescriptorCount); break;
-	case VK_DESCRIPTOR_TYPE_UNIFORM_BUFFER:            SET_DWORD_STAT(STAT_VulkanBindlessPeakUniformBuffer, State.PeakDescriptorCount); break;
-	case VK_DESCRIPTOR_TYPE_ACCELERATION_STRUCTURE_KHR:SET_DWORD_STAT(STAT_VulkanBindlessPeakAccelerationStructure, State.PeakDescriptorCount); break;
-	default: 
-		checkf(false, TEXT("Unknown descriptor type for stat update: %s"), VK_TYPE_TO_STRING(VkDescriptorType, DescriptorType));
-	}
-}
-
-FRHIDescriptorHandle FVulkanBindlessDescriptorManager::ReserveDescriptor(VkDescriptorType DescriptorType)
-{
-	if (bIsSupported)
-	{
-		const uint8 SetIndex = GetIndexForDescriptorType(DescriptorType);
-		const ERHIDescriptorType RHIType = GetUEDescriptorTypeFromVkDescriptorType(DescriptorType);
-
-		BindlessSetState& State = BindlessSetStates[SetIndex];
-		const uint32 ResourceIndex = GetFreeResourceIndex(State);
-		return FRHIDescriptorHandle(RHIType, ResourceIndex);
-	}
-	return FRHIDescriptorHandle();
-}
-
-void FVulkanBindlessDescriptorManager::UpdateDescriptor(FRHIDescriptorHandle DescriptorHandle, VkDescriptorType DescriptorType, VkDescriptorDataEXT DescriptorData, bool bImmediateUpdate)
-{
-	checkf(DescriptorHandle.IsValid(), TEXT("Attemping to update invalid descriptor handle!"));
-
-	const uint8 SetIndex = GetIndexForDescriptorType(DescriptorHandle.GetType());
-
-	BindlessSetState& State = BindlessSetStates[SetIndex];
-	const uint32 ByteOffset = DescriptorHandle.GetIndex() * State.DescriptorSize;
-	checkSlow(State.DescriptorTypes.Contains(DescriptorType));
-
-	VkDescriptorGetInfoEXT Info;
-	ZeroVulkanStruct(Info, VK_STRUCTURE_TYPE_DESCRIPTOR_GET_INFO_EXT);
-	Info.type = DescriptorType;
-	Info.data = DescriptorData;
-	VulkanRHI::vkGetDescriptorEXT(Device.GetInstanceHandle(), &Info, State.DescriptorSize, &State.DebugDescriptors[ByteOffset]);
-
-	if (bImmediateUpdate)
-	{
-		FMemory::Memcpy(&State.MappedPointer[ByteOffset], &State.DebugDescriptors[ByteOffset], State.DescriptorSize);
-	}
-	else
-	{
-		check(!IsInRenderingThread() || FRHICommandListExecutor::GetImmediateCommandList().Bypass() || !IsRunningRHIInSeparateThread());
-		FVulkanCommandBuffer* CmdBuffer = Device.GetImmediateContext().GetActiveCmdBuffer();
-
-		// :todo-jn:  Hack to avoid barriers/copies in renderpasses
-		if (CmdBuffer->IsInsideRenderPass())
+#if VULKAN_VALIDATE_DESCRIPTORS_WRITTEN
+	auto GetVkDescriptorTypeString = [](VkDescriptorType Type)
 		{
-			FMemory::Memcpy(&State.MappedPointer[ByteOffset], &State.DebugDescriptors[ByteOffset], State.DescriptorSize);
-		}
-		else
-		{
-			VulkanRHI::FStagingBuffer* StagingBuffer = Device.GetStagingManager().AcquireBuffer(State.DescriptorSize);
-			FMemory::Memcpy(StagingBuffer->GetMappedPointer(), &State.DebugDescriptors[ByteOffset], State.DescriptorSize);
+			switch (Type)
 			{
-				VkMemoryBarrier2 MemoryBarrier;
-				ZeroVulkanStruct(MemoryBarrier, VK_STRUCTURE_TYPE_MEMORY_BARRIER_2);
-				MemoryBarrier.srcStageMask = VK_PIPELINE_STAGE_2_ALL_COMMANDS_BIT;
-				MemoryBarrier.srcAccessMask = VK_ACCESS_2_DESCRIPTOR_BUFFER_READ_BIT_EXT | VK_ACCESS_2_MEMORY_READ_BIT | VK_ACCESS_2_MEMORY_WRITE_BIT;
-				MemoryBarrier.dstStageMask = VK_PIPELINE_STAGE_2_TRANSFER_BIT;
-				MemoryBarrier.dstAccessMask = VK_ACCESS_2_TRANSFER_READ_BIT | VK_ACCESS_2_TRANSFER_WRITE_BIT;
-
-				VkDependencyInfo DependencyInfo;
-				ZeroVulkanStruct(DependencyInfo, VK_STRUCTURE_TYPE_DEPENDENCY_INFO);
-				DependencyInfo.memoryBarrierCount = 1;
-				DependencyInfo.pMemoryBarriers = &MemoryBarrier;
-				VulkanRHI::vkCmdPipelineBarrier2KHR(CmdBuffer->GetHandle(), &DependencyInfo);
-
-				VkBufferCopy Region = {};
-				Region.srcOffset = 0;
-				Region.dstOffset = ByteOffset;
-				Region.size = State.DescriptorSize;
-				VulkanRHI::vkCmdCopyBuffer(CmdBuffer->GetHandle(), StagingBuffer->GetHandle(), State.BufferHandle, 1, &Region);
-
-				MemoryBarrier.srcStageMask = VK_PIPELINE_STAGE_2_TRANSFER_BIT;
-				MemoryBarrier.srcAccessMask = VK_ACCESS_2_TRANSFER_READ_BIT | VK_ACCESS_2_TRANSFER_WRITE_BIT;
-				MemoryBarrier.dstStageMask = VK_PIPELINE_STAGE_2_ALL_COMMANDS_BIT;
-				MemoryBarrier.dstAccessMask = VK_ACCESS_2_DESCRIPTOR_BUFFER_READ_BIT_EXT;
-				VulkanRHI::vkCmdPipelineBarrier2KHR(CmdBuffer->GetHandle(), &DependencyInfo);
+				// + 19 to skip "VK_DESCRIPTOR_TYPE_"
+#define VKSWITCHCASE(x)	case x: return FString(&TEXT(#x)[19]);
+				VKSWITCHCASE(VK_DESCRIPTOR_TYPE_SAMPLER)
+					VKSWITCHCASE(VK_DESCRIPTOR_TYPE_COMBINED_IMAGE_SAMPLER)
+					VKSWITCHCASE(VK_DESCRIPTOR_TYPE_SAMPLED_IMAGE)
+					VKSWITCHCASE(VK_DESCRIPTOR_TYPE_STORAGE_IMAGE)
+					VKSWITCHCASE(VK_DESCRIPTOR_TYPE_UNIFORM_TEXEL_BUFFER)
+					VKSWITCHCASE(VK_DESCRIPTOR_TYPE_STORAGE_TEXEL_BUFFER)
+					VKSWITCHCASE(VK_DESCRIPTOR_TYPE_UNIFORM_BUFFER)
+					VKSWITCHCASE(VK_DESCRIPTOR_TYPE_STORAGE_BUFFER)
+					VKSWITCHCASE(VK_DESCRIPTOR_TYPE_UNIFORM_BUFFER_DYNAMIC)
+					VKSWITCHCASE(VK_DESCRIPTOR_TYPE_STORAGE_BUFFER_DYNAMIC)
+					VKSWITCHCASE(VK_DESCRIPTOR_TYPE_INPUT_ATTACHMENT)
+#undef VKSWITCHCASE
+			default:
+				break;
 			}
-			Device.GetStagingManager().ReleaseBuffer(&Device.GetImmediateContext(), StagingBuffer);
-		}
-	}
 
-	UpdateStatsForHandle(DescriptorType);
-}
+			return FString::Printf(TEXT("Unknown VkDescriptorType %d"), (int32)Type);
+		};
 
-void FVulkanBindlessDescriptorManager::UpdateSampler(FRHIDescriptorHandle DescriptorHandle, VkSampler VulkanSampler)
-{
-	if (bIsSupported)
-	{
-		VkDescriptorDataEXT DescriptorData;
-		DescriptorData.pSampler = &VulkanSampler;
-		UpdateDescriptor(DescriptorHandle, VK_DESCRIPTOR_TYPE_SAMPLER, DescriptorData, true);
-	}
-}
+	const uint32 Writes = NumWrites;
+	if (Writes == 0)
+		return;
 
-void FVulkanBindlessDescriptorManager::UpdateImage(FRHIDescriptorHandle DescriptorHandle, VkDescriptorType DescriptorType, VkImageView ImageView, bool bIsDepthStencil, bool bImmediateUpdate)
-{
-	if (bIsSupported)
+	bool bFail = false;
+	if (Writes <= 32) //early out for the most common case.
 	{
-		check((DescriptorType == VK_DESCRIPTOR_TYPE_SAMPLED_IMAGE) || (DescriptorType == VK_DESCRIPTOR_TYPE_STORAGE_IMAGE));
-
-		VkDescriptorImageInfo DescriptorImageInfo;
-		DescriptorImageInfo.sampler = VK_NULL_HANDLE;
-		DescriptorImageInfo.imageView = ImageView;
-		DescriptorImageInfo.imageLayout = (DescriptorType == VK_DESCRIPTOR_TYPE_STORAGE_IMAGE) ? VK_IMAGE_LAYOUT_GENERAL :
-			(bIsDepthStencil ? VK_IMAGE_LAYOUT_READ_ONLY_OPTIMAL : VK_IMAGE_LAYOUT_SHADER_READ_ONLY_OPTIMAL);
-
-		VkDescriptorDataEXT DescriptorData;
-		DescriptorData.pSampledImage = &DescriptorImageInfo;  // same pointer for storage, it's a union
-		UpdateDescriptor(DescriptorHandle, DescriptorType, DescriptorData, bImmediateUpdate);
+		bFail = WrittenMask[0] != ((1llu << Writes) - 1);
 	}
-}
-
-void FVulkanBindlessDescriptorManager::UpdateBuffer(FRHIDescriptorHandle DescriptorHandle, VkDescriptorType DescriptorType, VkBuffer VulkanBuffer, VkDeviceSize BufferOffset, VkDeviceSize BufferSize, bool bImmediateUpdate)
-{
-	if (bIsSupported)
+	else
 	{
-		check((DescriptorType == VK_DESCRIPTOR_TYPE_UNIFORM_BUFFER) || (DescriptorType == VK_DESCRIPTOR_TYPE_STORAGE_BUFFER));
-
-		VkBufferDeviceAddressInfo BufferInfo;
-		ZeroVulkanStruct(BufferInfo, VK_STRUCTURE_TYPE_BUFFER_DEVICE_ADDRESS_INFO);
-		BufferInfo.buffer = VulkanBuffer;
-		const VkDeviceAddress BufferAddress = VulkanRHI::vkGetBufferDeviceAddressKHR(Device.GetInstanceHandle(), &BufferInfo);
+		const int32 Last = int32(WrittenMask.Num() - 1);
+		for (int32 i = 0; !bFail && i < Last; ++i)
+		{
+			uint64 Mask = WrittenMask[i];
+			bFail = bFail || Mask != 0xffffffff;
+		}
 
-		UpdateBuffer(DescriptorHandle, DescriptorType, BufferAddress + BufferOffset, BufferSize, bImmediateUpdate);
+		const uint32 TailCount = Writes - (Last * 32);
+		check(TailCount != 0);
+		const uint32 TailMask = (1llu << TailCount) - 1;
+		bFail = bFail || TailMask != WrittenMask[Last];
 	}
-}
 
-void FVulkanBindlessDescriptorManager::UpdateBuffer(FRHIDescriptorHandle DescriptorHandle, VkDescriptorType DescriptorType, VkDeviceAddress BufferAddress, VkDeviceSize BufferSize, bool bImmediateUpdate)
-{
-	if (bIsSupported)
+	if (bFail)
 	{
-		check((DescriptorType == VK_DESCRIPTOR_TYPE_UNIFORM_BUFFER) || (DescriptorType == VK_DESCRIPTOR_TYPE_STORAGE_BUFFER));
-
-		VkDescriptorAddressInfoEXT AddressInfo;
-		ZeroVulkanStruct(AddressInfo, VK_STRUCTURE_TYPE_DESCRIPTOR_ADDRESS_INFO_EXT);
-		AddressInfo.address = BufferAddress;
-		AddressInfo.range = BufferSize;
-
-		VkDescriptorDataEXT DescriptorData;
-		DescriptorData.pStorageBuffer = &AddressInfo;  // same pointer for uniform, it's a union
-		UpdateDescriptor(DescriptorHandle, DescriptorType, DescriptorData, bImmediateUpdate);
+		FString Descriptors;
+		for (uint32 i = 0; i < Writes; ++i)
+		{
+			uint32 Index = i / 32;
+			uint32 Mask = i % 32;
+			if (0 == (WrittenMask[Index] & (1llu << Mask)))
+			{
+				FString TypeString = GetVkDescriptorTypeString(WriteDescriptors[i].descriptorType);
+				Descriptors += FString::Printf(TEXT("\t\tDescriptorWrite %d/%d Was not written(Type %s)\n"), i, NumWrites, *TypeString);
+			}
+		}
+		UE_LOG(LogVulkanRHI, Warning, TEXT("Not All descriptors where filled out. this can/will cause a driver crash\n%s\n"), *Descriptors);
+		ensureMsgf(false, TEXT("Not All descriptors where filled out. this can/will cause a driver crash\n%s\n"), *Descriptors);
 	}
+#endif
 }
 
-void FVulkanBindlessDescriptorManager::UpdateTexelBuffer(FRHIDescriptorHandle DescriptorHandle, VkDescriptorType DescriptorType, const VkBufferViewCreateInfo& ViewInfo, bool bImmediateUpdate)
+void FVulkanDescriptorSetWriter::Reset()
 {
-	if (bIsSupported)
-	{
-		check((DescriptorType == VK_DESCRIPTOR_TYPE_UNIFORM_TEXEL_BUFFER) || (DescriptorType == VK_DESCRIPTOR_TYPE_STORAGE_TEXEL_BUFFER));
-
-		// :todo-jn: start caching buffer addresses in resources to avoid the extra call
-		VkBufferDeviceAddressInfo BufferInfo;
-		ZeroVulkanStruct(BufferInfo, VK_STRUCTURE_TYPE_BUFFER_DEVICE_ADDRESS_INFO);
-		BufferInfo.buffer = ViewInfo.buffer;
-		const VkDeviceAddress BufferAddress = VulkanRHI::vkGetBufferDeviceAddressKHR(Device.GetInstanceHandle(), &BufferInfo);
-
-		VkDescriptorAddressInfoEXT AddressInfo;
-		ZeroVulkanStruct(AddressInfo, VK_STRUCTURE_TYPE_DESCRIPTOR_ADDRESS_INFO_EXT);
-		AddressInfo.address = BufferAddress + ViewInfo.offset;
-		AddressInfo.range = ViewInfo.range;
-		AddressInfo.format = ViewInfo.format;
-
-		VkDescriptorDataEXT DescriptorData;
-		DescriptorData.pUniformTexelBuffer = &AddressInfo;  // same pointer for storage, it's a union
-		UpdateDescriptor(DescriptorHandle, DescriptorType, DescriptorData, bImmediateUpdate);
-	}
-}
+	bHasVolatileResources = false;
 
-void FVulkanBindlessDescriptorManager::UpdateAccelerationStructure(FRHIDescriptorHandle DescriptorHandle, VkAccelerationStructureKHR AccelerationStructure, bool bImmediateUpdate)
-{
-	if (bIsSupported)
-	{
-		// :todo-jn: start caching AccelerationStructure in resources to avoid the extra call
-		VkAccelerationStructureDeviceAddressInfoKHR AccelerationStructureDeviceAddressInfo;
-		ZeroVulkanStruct(AccelerationStructureDeviceAddressInfo, VK_STRUCTURE_TYPE_ACCELERATION_STRUCTURE_DEVICE_ADDRESS_INFO_KHR);
-		AccelerationStructureDeviceAddressInfo.accelerationStructure = AccelerationStructure;
-		const VkDeviceAddress BufferAddress = VulkanRHI::vkGetAccelerationStructureDeviceAddressKHR(Device.GetInstanceHandle(), &AccelerationStructureDeviceAddressInfo);
-
-		VkDescriptorDataEXT DescriptorData;
-		DescriptorData.accelerationStructure = BufferAddress;
-		UpdateDescriptor(DescriptorHandle, VK_DESCRIPTOR_TYPE_ACCELERATION_STRUCTURE_KHR, DescriptorData, bImmediateUpdate);
-	}
+#if VULKAN_VALIDATE_DESCRIPTORS_WRITTEN
+	WrittenMask = BaseWrittenMask;
+#endif
 }
-
-
-uint32 FVulkanBindlessDescriptorManager::GetFreeResourceIndex(FVulkanBindlessDescriptorManager::BindlessSetState& State)
+void FVulkanDescriptorSetWriter::SetWritten(uint32 DescriptorIndex)
 {
-	INC_DWORD_STAT(STAT_VulkanBindlessWritePerFrame);
-
-	{
-		FScopeLock ScopeLock(&State.FreeListCS);
-		if ((State.FreeListHead != MAX_uint32) && (State.PeakDescriptorCount >= State.MaxDescriptorCount)) // todo-jn: temp
-		{
-			const uint32 FreeIndex = State.FreeListHead;
-			const uint32 ByteOffset = State.FreeListHead * State.DescriptorSize;
-			uint32* NextSlot = (uint32*)(&State.DebugDescriptors[ByteOffset]);
-			State.FreeListHead = *NextSlot;
-			return FreeIndex;
-		}
-	}
-
-	INC_DWORD_STAT(STAT_VulkanBindlessPeakDescriptorCount);
-
-	const uint32 ResourceIndex = State.PeakDescriptorCount++;
-	checkf(ResourceIndex < State.MaxDescriptorCount, TEXT("You need to grow the resource array size for [%s]!"), VK_TYPE_TO_STRING(VkDescriptorType, State.DescriptorTypes[0]));
-	return ResourceIndex;
+#if VULKAN_VALIDATE_DESCRIPTORS_WRITTEN
+	uint32 Index = DescriptorIndex / 32;
+	uint32 Mask = DescriptorIndex % 32;
+	WrittenMask[Index] |= (1 << Mask);
+#endif
 }
-
-void FVulkanBindlessDescriptorManager::Unregister(FRHIDescriptorHandle DescriptorHandle)
+void FVulkanDescriptorSetWriter::SetWrittenBase(uint32 DescriptorIndex)
 {
-	if (DescriptorHandle.IsValid())
-	{
-		checkf(bIsSupported, TEXT("Unregistering a valid handle but bindless is not supported!"));
-
-		const uint8 SetIndex = GetIndexForDescriptorType(DescriptorHandle.GetType());
-		BindlessSetState& State = BindlessSetStates[SetIndex];
-
-		FScopeLock ScopeLock(&State.FreeListCS);
-
-		const uint32 PreviousHead = State.FreeListHead;
-		State.FreeListHead = DescriptorHandle.GetIndex();
-		const uint32 ByteOffset = DescriptorHandle.GetIndex() * State.DescriptorSize;
-		uint32* NewSlot = (uint32*)(&State.DebugDescriptors[ByteOffset]);
-		FMemory::Memzero(NewSlot, State.DescriptorSize); // easier for debugging for now
-		*NewSlot = PreviousHead;
-
-		// Clear the descriptor
-		// todo-jn: invalidate the GPU side?
-	}
+#if VULKAN_VALIDATE_DESCRIPTORS_WRITTEN	
+	uint32 Index = DescriptorIndex / 32;
+	uint32 Mask = DescriptorIndex % 32;
+	BaseWrittenMask[Index] |= (1 << Mask);
+#endif
 }
 
-void FVulkanBindlessDescriptorManager::UpdateUBAllocator()
+void FVulkanDescriptorSetWriter::InitWrittenMasks(uint32 NumDescriptorWrites)
 {
-	if (bIsSupported && SingleUseUBAllocator)
-	{
-		SingleUseUBAllocator->UpdateBlocks();
-	}
+#if VULKAN_VALIDATE_DESCRIPTORS_WRITTEN	
+	uint32 Size = (NumDescriptorWrites + 31) / 32;
+	WrittenMask.Empty(Size);
+	WrittenMask.SetNumZeroed(Size);
+	BaseWrittenMask.Empty(Size);
+	BaseWrittenMask.SetNumZeroed(Size);
+#endif
 }
diff --git a/Engine/Source/Runtime/VulkanRHI/Private/VulkanDescriptorSets.h b/Engine/Source/Runtime/VulkanRHI/Private/VulkanDescriptorSets.h
--- a/Engine/Source/Runtime/VulkanRHI/Private/VulkanDescriptorSets.h
+++ b/Engine/Source/Runtime/VulkanRHI/Private/VulkanDescriptorSets.h
@@ -1,9 +1,5 @@
 // Copyright Epic Games, Inc. All Rights Reserved.
 
-/*=============================================================================
-	VulkanPipelineState.h: Vulkan pipeline state definitions.
-=============================================================================*/
-
 #pragma once
 
 #include "Async/AsyncWork.h"
@@ -28,12 +24,7 @@ class FVulkanDSetsKey : public TDataKey<FVulkanDSetsKey, true> {};
 
 struct FUniformBufferGatherInfo
 {
-	FUniformBufferGatherInfo()
-	{
-		FMemory::Memzero(CodeHeaders);
-	}
-
-	const FVulkanShaderHeader* CodeHeaders[ShaderStage::MaxNumStages];
+	const FVulkanShaderHeader* CodeHeaders[ShaderStage::MaxNumStages]{};
 };
 
 
@@ -844,151 +835,6 @@ protected:
 };
 
 
-
-// Manager for resource descriptors used in bindless rendering.
-class FVulkanBindlessDescriptorManager
-{
-public:
-	FVulkanBindlessDescriptorManager(FVulkanDevice& InDevice);
-	~FVulkanBindlessDescriptorManager();
-
-	typedef TStaticArray<TArray<VkDescriptorAddressInfoEXT>, ShaderStage::MaxNumStages> FUniformBufferDescriptorArrays;
-
-	void Init();
-	void Deinit();
-
-	static bool VerifySupport(FVulkanDevice& InDevice);
-	inline bool IsSupported()
-	{
-		return bIsSupported;
-	}
-
-	inline VkPipelineLayout GetPipelineLayout() const
-	{
-		return BindlessPipelineLayout;
-	}
-
-	void BindDescriptorBuffers(VkCommandBuffer CommandBuffer, VkPipelineStageFlags SupportedStages);
-
-	FRHIDescriptorHandle ReserveDescriptor(VkDescriptorType DescriptorType);
-
-	void UpdateSampler(FRHIDescriptorHandle DescriptorHandle, VkSampler VulkanSampler);
-	void UpdateImage(FRHIDescriptorHandle DescriptorHandle, VkDescriptorType DescriptorType, VkImageView VulkanImage, bool bIsDepthStencil, bool bImmediateUpdate = true);
-	void UpdateBuffer(FRHIDescriptorHandle DescriptorHandle, VkDescriptorType DescriptorType, VkBuffer VulkanBuffer, VkDeviceSize BufferOffset, VkDeviceSize BufferSize, bool bImmediateUpdate = true);
-	void UpdateBuffer(FRHIDescriptorHandle DescriptorHandle, VkDescriptorType DescriptorType, VkDeviceAddress BufferAddress, VkDeviceSize BufferSize, bool bImmediateUpdate = true);
-	void UpdateTexelBuffer(FRHIDescriptorHandle DescriptorHandle, VkDescriptorType DescriptorType, const VkBufferViewCreateInfo& ViewInfo, bool bImmediateUpdate = true);
-	void UpdateAccelerationStructure(FRHIDescriptorHandle DescriptorHandle, VkAccelerationStructureKHR AccelerationStructure, bool bImmediateUpdate = true);
-
-	void RegisterUniformBuffers(FVulkanCommandListContext& Context, VkPipelineBindPoint BindPoint, const FUniformBufferDescriptorArrays& StageUBs);
-
-	void Unregister(FRHIDescriptorHandle DescriptorHandle);
-
-	void UpdateUBAllocator();
-
-private:
-	FVulkanDevice& Device;
-	const bool bIsSupported;
-
-	VkDescriptorSetLayout EmptyDescriptorSetLayout = VK_NULL_HANDLE;
-
-	struct BindlessSetState
-	{
-		TArray<VkDescriptorType> DescriptorTypes;
-		uint32 MaxDescriptorCount = 0;
-		std::atomic<uint32> PeakDescriptorCount = 1;  // always keep a null descriptor in slot 0
-		VkDescriptorSetLayout DescriptorSetLayout = VK_NULL_HANDLE;
-
-		FCriticalSection FreeListCS;
-		uint32 FreeListHead = MAX_uint32;
-
-		uint32 DescriptorSize = 0;
-		VkBuffer BufferHandle = VK_NULL_HANDLE;
-		VkDeviceMemory MemoryHandle = VK_NULL_HANDLE;
-		uint8* MappedPointer = nullptr;
-
-		TArray<uint8> DebugDescriptors;
-	};
-	BindlessSetState BindlessSetStates[VulkanBindless::NumBindlessSets];
-
-	VkDescriptorSetLayout SingleUseUBDescriptorSetLayout = VK_NULL_HANDLE;
-	VulkanRHI::FTempBlockAllocator* SingleUseUBAllocator = nullptr;
-
-	VkDescriptorBufferBindingInfoEXT BufferBindingInfo[VulkanBindless::NumBindlessSets];
-	uint32_t BufferIndices[VulkanBindless::MaxNumSets];
-
-	VkPipelineLayout BindlessPipelineLayout = VK_NULL_HANDLE;
-
-	uint32 GetFreeResourceIndex(BindlessSetState& Desc);
-	void UpdateStatsForHandle(VkDescriptorType DescriptorType);
-	void UpdateDescriptor(FRHIDescriptorHandle DescriptorHandle, VkDescriptorType DescriptorType, VkDescriptorDataEXT DescriptorData, bool bImmediateUpdate);
-};
-
-
-
-// Layout for a Pipeline, also includes DescriptorSets layout
-class FVulkanLayout : public VulkanRHI::FDeviceChild
-{
-public:
-	FVulkanLayout(FVulkanDevice* InDevice, bool InGfxLayout, bool InUsesBindless);
-	virtual ~FVulkanLayout();
-
-	bool IsGfxLayout() const
-	{
-		return bIsGfxLayout;
-	}
-
-	inline const FVulkanDescriptorSetsLayout& GetDescriptorSetsLayout() const
-	{
-		return DescriptorSetLayout;
-	}
-
-	inline VkPipelineLayout GetPipelineLayout() const
-	{
-		return bUsesBindless ? Device->GetBindlessDescriptorManager()->GetPipelineLayout() : PipelineLayout;
-	}
-
-	inline bool HasDescriptors() const
-	{
-		return DescriptorSetLayout.GetLayouts().Num() > 0;
-	}
-
-	inline uint32 GetDescriptorSetLayoutHash() const
-	{
-		return DescriptorSetLayout.GetHash();
-	}
-
-protected:
-	const bool bIsGfxLayout;
-	const bool bUsesBindless;
-	FVulkanDescriptorSetsLayout	DescriptorSetLayout;
-	VkPipelineLayout			PipelineLayout;
-
-	template <bool bIsCompute>
-	inline void FinalizeBindings(const FUniformBufferGatherInfo& UBGatherInfo)
-	{
-		// Setting descriptor is only allowed prior to compiling the layout
-		check(DescriptorSetLayout.GetHandles().Num() == 0);
-
-		DescriptorSetLayout.FinalizeBindings<bIsCompute>(UBGatherInfo);
-	}
-
-	inline void ProcessBindingsForStage(VkShaderStageFlagBits StageFlags, ShaderStage::EStage DescSet, const FVulkanShaderHeader& CodeHeader, FUniformBufferGatherInfo& OutUBGatherInfo) const
-	{
-		// Setting descriptor is only allowed prior to compiling the layout
-		check(DescriptorSetLayout.GetHandles().Num() == 0);
-
-		DescriptorSetLayout.ProcessBindingsForStage(StageFlags, DescSet, CodeHeader, OutUBGatherInfo);
-	}
-
-	void Compile(FVulkanDescriptorSetLayoutMap& DSetLayoutMap);
-
-	friend class FVulkanComputePipeline;
-	friend class FVulkanGfxPipeline;
-	friend class FVulkanPipelineStateCacheManager;
-	friend class FVulkanRayTracingPipelineState;
-};
-
-
 class FVulkanGenericDescriptorPool : FNoncopyable
 {
 public:
@@ -1078,7 +924,3 @@ private:
 	TUniquePtr<FCachedPool> FreePool;
 	float PoolAllocRatio;
 };
-
-
-
-
diff --git a/Engine/Source/Runtime/VulkanRHI/Private/VulkanDevice.cpp b/Engine/Source/Runtime/VulkanRHI/Private/VulkanDevice.cpp
--- a/Engine/Source/Runtime/VulkanRHI/Private/VulkanDevice.cpp
+++ b/Engine/Source/Runtime/VulkanRHI/Private/VulkanDevice.cpp
@@ -20,6 +20,7 @@
 #include "VulkanRayTracing.h"
 #include "VulkanDescriptorSets.h"
 #include "VulkanChunkedPipelineCache.h"
+#include "VulkanBindlessDescriptorManager.h"
 
 static TAutoConsoleVariable<int32> GRHIAllowAsyncComputeCvar(
 	TEXT("r.Vulkan.AllowAsyncCompute"),
diff --git a/Engine/Source/Runtime/VulkanRHI/Private/VulkanFramebuffer.cpp b/Engine/Source/Runtime/VulkanRHI/Private/VulkanFramebuffer.cpp
new file mode 100644
--- /dev/null
+++ b/Engine/Source/Runtime/VulkanRHI/Private/VulkanFramebuffer.cpp
@@ -0,0 +1,398 @@
+// Copyright Epic Games, Inc. All Rights Reserved.
+
+#include "VulkanFramebuffer.h"
+#include "VulkanDevice.h"
+#include "VulkanRenderpass.h"
+#include "VulkanRenderTargetLayout.h"
+
+FVulkanFramebuffer::FVulkanFramebuffer(FVulkanDevice& Device, const FRHISetRenderTargetsInfo& InRTInfo, const FVulkanRenderTargetLayout& RTLayout, const FVulkanRenderPass& RenderPass)
+	: Framebuffer(VK_NULL_HANDLE)
+	, NumColorRenderTargets(InRTInfo.NumColorRenderTargets)
+	, NumColorAttachments(0)
+	, DepthStencilRenderTargetImage(VK_NULL_HANDLE)
+	, DepthStencilResolveRenderTargetImage(VK_NULL_HANDLE)
+	, FragmentDensityImage(VK_NULL_HANDLE)
+{
+	FMemory::Memzero(ColorRenderTargetImages);
+	FMemory::Memzero(ColorResolveTargetImages);
+
+	AttachmentTextureViews.Empty(RTLayout.GetNumAttachmentDescriptions());
+
+	uint32 MipIndex = 0;
+
+	const VkExtent3D& RTExtents = RTLayout.GetExtent3D();
+	// Adreno does not like zero size RTs
+	check(RTExtents.width != 0 && RTExtents.height != 0);
+	uint32 NumLayers = RTExtents.depth;
+
+	for (int32 Index = 0; Index < InRTInfo.NumColorRenderTargets; ++Index)
+	{
+		FRHITexture* RHITexture = InRTInfo.ColorRenderTarget[Index].Texture;
+		if (!RHITexture)
+		{
+			continue;
+		}
+
+		FVulkanTexture* Texture = ResourceCast(RHITexture);
+
+		ColorRenderTargetImages[Index] = Texture->Image;
+		MipIndex = InRTInfo.ColorRenderTarget[Index].MipIndex;
+
+		FVulkanView* View = GetColorRenderTargetViewDesc(Texture, MipIndex, InRTInfo.ColorRenderTarget[Index].ArraySliceIndex, RTLayout.GetMultiViewCount(), NumLayers);
+		AttachmentTextureViews.Add(View);
+
+		++NumColorAttachments;
+
+		// Check the RTLayout as well to make sure the resolve attachment is needed (Vulkan and Feature level specific)
+		// See: FVulkanRenderTargetLayout constructor with FRHIRenderPassInfo
+		if (InRTInfo.bHasResolveAttachments && RTLayout.GetHasResolveAttachments() && RTLayout.GetResolveAttachmentReferences()[Index].layout != VK_IMAGE_LAYOUT_UNDEFINED)
+		{
+			FRHITexture* ResolveRHITexture = InRTInfo.ColorResolveRenderTarget[Index].Texture;
+			FVulkanTexture* ResolveTexture = ResourceCast(ResolveRHITexture);
+			ColorResolveTargetImages[Index] = ResolveTexture->Image;
+
+			if (FVulkanView* ResolveView = GetColorResolveTargetViewDesc(ResolveTexture, MipIndex, InRTInfo.ColorRenderTarget[Index].ArraySliceIndex))
+			{
+				AttachmentTextureViews.Add(ResolveView);
+			}
+		}
+	}
+
+	if (RTLayout.GetHasDepthStencil())
+	{
+		FVulkanTexture* Texture = ResourceCast(InRTInfo.DepthStencilRenderTarget.Texture);
+		const FRHITextureDesc& Desc = Texture->GetDesc();
+		DepthStencilRenderTargetImage = Texture->Image;
+
+		check(Texture->PartialView);
+		PartialDepthTextureView = Texture->PartialView;
+
+		FVulkanView* View = GetDepthStencilTargetViewDesc(Texture, InRTInfo.NumColorRenderTargets, MipIndex, NumLayers);
+		AttachmentTextureViews.Add(View);
+
+		if (RTLayout.GetHasDepthStencilResolve() && RTLayout.GetDepthStencilResolveAttachmentReference()->layout != VK_IMAGE_LAYOUT_UNDEFINED)
+		{
+			FRHITexture* ResolveRHITexture = InRTInfo.DepthStencilResolveRenderTarget.Texture;
+			FVulkanTexture* ResolveTexture = ResourceCast(ResolveRHITexture);
+			DepthStencilResolveRenderTargetImage = ResolveTexture->Image;
+
+			if (FVulkanView* ResolveView = GetDepthStencilResolveTargetViewDesc(ResolveTexture, MipIndex))
+			{
+				AttachmentTextureViews.Add(ResolveView);
+			}
+		}
+	}
+
+	if (GRHISupportsAttachmentVariableRateShading && RTLayout.GetHasFragmentDensityAttachment())
+	{
+		FVulkanTexture* Texture = ResourceCast(InRTInfo.ShadingRateTexture);
+		FragmentDensityImage = Texture->Image;
+		FVulkanView* View = GetFragmentDensityAttachmentViewDesc(Texture, MipIndex);
+		AttachmentTextureViews.Add(View);
+	}
+
+	TArray<VkImageView> AttachmentViews;
+	AttachmentViews.Reserve(AttachmentTextureViews.Num());
+	for (FVulkanView const* View : AttachmentTextureViews)
+	{
+		AttachmentViews.Add(View->GetTextureView().View);
+	}
+
+	VkFramebufferCreateInfo CreateInfo;
+	ZeroVulkanStruct(CreateInfo, VK_STRUCTURE_TYPE_FRAMEBUFFER_CREATE_INFO);
+	CreateInfo.renderPass = RenderPass.GetHandle();
+	CreateInfo.attachmentCount = AttachmentViews.Num();
+	CreateInfo.pAttachments = AttachmentViews.GetData();
+	CreateInfo.width = RTExtents.width;
+	CreateInfo.height = RTExtents.height;
+	CreateInfo.layers = NumLayers;
+
+	VERIFYVULKANRESULT_EXPANDED(VulkanRHI::vkCreateFramebuffer(Device.GetInstanceHandle(), &CreateInfo, VULKAN_CPU_ALLOCATOR, &Framebuffer));
+
+	RenderArea.offset.x = 0;
+	RenderArea.offset.y = 0;
+	RenderArea.extent.width = RTExtents.width;
+	RenderArea.extent.height = RTExtents.height;
+
+	INC_DWORD_STAT(STAT_VulkanNumFrameBuffers);
+}
+
+FVulkanFramebuffer::~FVulkanFramebuffer()
+{
+	ensure(Framebuffer == VK_NULL_HANDLE);
+}
+
+
+FVulkanView* FVulkanFramebuffer::GetColorRenderTargetViewDesc(FVulkanTexture* Texture, uint32 MipIndex, int32 SliceIndex, uint32 MultiViewCount, uint32& InOutNumLayers)
+{
+	check(Texture->Image != VK_NULL_HANDLE);
+
+	const FRHITextureDesc& Desc = Texture->GetDesc();
+	const VkImageUsageFlags ImageUsageFlags = VK_IMAGE_USAGE_COLOR_ATTACHMENT_BIT | (Texture->ImageUsageFlags & VK_IMAGE_USAGE_INPUT_ATTACHMENT_BIT);
+
+	if (Texture->GetViewType() == VK_IMAGE_VIEW_TYPE_2D || Texture->GetViewType() == VK_IMAGE_VIEW_TYPE_2D_ARRAY)
+	{
+		uint32 ArraySliceIndex = 0;
+		uint32 NumArraySlices = 1;
+		if (SliceIndex == -1)
+		{
+			ArraySliceIndex = 0;
+			NumArraySlices = Texture->GetNumberOfArrayLevels();
+		}
+		else
+		{
+			ArraySliceIndex = SliceIndex;
+			NumArraySlices = 1;
+			check(ArraySliceIndex < Texture->GetNumberOfArrayLevels());
+		}
+
+		// About !RTLayout.GetIsMultiView(), from https://registry.khronos.org/vulkan/specs/1.3-extensions/man/html/VkFramebufferCreateInfo.html: 
+		// If the render pass uses multiview, then layers must be one
+		if (Texture->GetViewType() == VK_IMAGE_VIEW_TYPE_2D_ARRAY && (MultiViewCount == 0))
+		{
+			InOutNumLayers = NumArraySlices;
+		}
+
+		FVulkanTextureViewDesc ViewDesc;
+		ViewDesc.ViewType = Texture->GetViewType();
+		ViewDesc.AspectFlags = Texture->GetFullAspectMask();
+		ViewDesc.UEFormat = Desc.Format;
+		ViewDesc.Format = Texture->ViewFormat;
+		ViewDesc.FirstMip = MipIndex;
+		ViewDesc.NumMips = 1;
+		ViewDesc.ArraySliceIndex = ArraySliceIndex;
+		ViewDesc.NumArraySlices = NumArraySlices;
+		ViewDesc.bUseIdentitySwizzle = true;
+		ViewDesc.ImageUsageFlags = ImageUsageFlags;
+		return Texture->FindOrAddInternalView(ViewDesc);
+	}
+	else if (Texture->GetViewType() == VK_IMAGE_VIEW_TYPE_CUBE || Texture->GetViewType() == VK_IMAGE_VIEW_TYPE_CUBE_ARRAY)
+	{
+		// Cube always renders one face at a time
+		INC_DWORD_STAT(STAT_VulkanNumImageViews);
+
+		FVulkanTextureViewDesc ViewDesc;
+		ViewDesc.ViewType = VK_IMAGE_VIEW_TYPE_2D;
+		ViewDesc.AspectFlags = Texture->GetFullAspectMask();
+		ViewDesc.UEFormat = Desc.Format;
+		ViewDesc.Format = Texture->ViewFormat;
+		ViewDesc.FirstMip = MipIndex;
+		ViewDesc.NumMips = 1;
+		ViewDesc.ArraySliceIndex = SliceIndex;
+		ViewDesc.NumArraySlices = 1;
+		ViewDesc.bUseIdentitySwizzle = true;
+		ViewDesc.ImageUsageFlags = ImageUsageFlags;
+		return Texture->FindOrAddInternalView(ViewDesc);
+	}
+	else if (Texture->GetViewType() == VK_IMAGE_VIEW_TYPE_3D)
+	{
+		FVulkanTextureViewDesc ViewDesc;
+		ViewDesc.ViewType = VK_IMAGE_VIEW_TYPE_2D_ARRAY;
+		ViewDesc.AspectFlags = Texture->GetFullAspectMask();
+		ViewDesc.UEFormat = Desc.Format;
+		ViewDesc.Format = Texture->ViewFormat;
+		ViewDesc.FirstMip = MipIndex;
+		ViewDesc.NumMips = 1;
+		ViewDesc.ArraySliceIndex = 0;
+		ViewDesc.NumArraySlices = Desc.Depth;
+		ViewDesc.bUseIdentitySwizzle = true;
+		ViewDesc.ImageUsageFlags = ImageUsageFlags;
+		return Texture->FindOrAddInternalView(ViewDesc);
+	}
+	ensure(0);
+	return nullptr;
+}
+
+FVulkanView* FVulkanFramebuffer::GetColorResolveTargetViewDesc(FVulkanTexture* ResolveTexture, uint32 MipIndex, int32 ArraySliceIndex)
+{
+	//resolve attachments only supported for 2d/2d array textures
+	if (ResolveTexture->GetViewType() == VK_IMAGE_VIEW_TYPE_2D || ResolveTexture->GetViewType() == VK_IMAGE_VIEW_TYPE_2D_ARRAY)
+	{
+		FVulkanTextureViewDesc ViewDesc;
+		ViewDesc.ViewType = ResolveTexture->GetViewType();
+		ViewDesc.AspectFlags = ResolveTexture->GetFullAspectMask();
+		ViewDesc.UEFormat = ResolveTexture->GetDesc().Format;
+		ViewDesc.Format = ResolveTexture->ViewFormat;
+		ViewDesc.FirstMip = MipIndex;
+		ViewDesc.NumMips = 1;
+		ViewDesc.ArraySliceIndex = FMath::Max(0, ArraySliceIndex);
+		ViewDesc.NumArraySlices = ResolveTexture->GetNumberOfArrayLevels();
+		ViewDesc.bUseIdentitySwizzle = true;
+		return ResolveTexture->FindOrAddInternalView(ViewDesc);
+	}
+
+	return nullptr;
+}
+
+FVulkanView* FVulkanFramebuffer::GetDepthStencilTargetViewDesc(FVulkanTexture* Texture, uint32 NumColorAttachments, uint32 MipIndex, uint32& InOutNumLayers)
+{
+	ensure(Texture->GetViewType() == VK_IMAGE_VIEW_TYPE_2D || Texture->GetViewType() == VK_IMAGE_VIEW_TYPE_2D_ARRAY || Texture->GetViewType() == VK_IMAGE_VIEW_TYPE_CUBE);
+	if ((NumColorAttachments == 0) && (Texture->GetViewType() == VK_IMAGE_VIEW_TYPE_CUBE))
+	{
+		InOutNumLayers = 6;
+
+		FVulkanTextureViewDesc ViewDesc;
+		ViewDesc.ViewType = VK_IMAGE_VIEW_TYPE_2D_ARRAY;
+		ViewDesc.AspectFlags = Texture->GetFullAspectMask();
+		ViewDesc.UEFormat = Texture->GetDesc().Format;
+		ViewDesc.Format = Texture->ViewFormat;
+		ViewDesc.FirstMip = MipIndex;
+		ViewDesc.NumMips = 1;
+		ViewDesc.ArraySliceIndex = 0;
+		ViewDesc.NumArraySlices = 6;
+		ViewDesc.bUseIdentitySwizzle = true;
+		return Texture->FindOrAddInternalView(ViewDesc);
+	}
+	else if ((Texture->GetViewType() == VK_IMAGE_VIEW_TYPE_2D) || (Texture->GetViewType() == VK_IMAGE_VIEW_TYPE_2D_ARRAY))
+	{
+		// depth attachments need a separate view to have no swizzle components, for validation correctness
+		FVulkanTextureViewDesc ViewDesc;
+		ViewDesc.ViewType = Texture->GetViewType();
+		ViewDesc.AspectFlags = Texture->GetFullAspectMask();
+		ViewDesc.UEFormat = Texture->GetDesc().Format;
+		ViewDesc.Format = Texture->ViewFormat;
+		ViewDesc.FirstMip = MipIndex;
+		ViewDesc.NumMips = 1;
+		ViewDesc.ArraySliceIndex = 0;
+		ViewDesc.NumArraySlices = Texture->GetNumberOfArrayLevels();
+		ViewDesc.bUseIdentitySwizzle = true;
+		return Texture->FindOrAddInternalView(ViewDesc);
+	}
+	else
+	{
+		return Texture->DefaultView;
+	}
+}
+
+FVulkanView* FVulkanFramebuffer::GetDepthStencilResolveTargetViewDesc(FVulkanTexture* ResolveTexture, uint32 MipIndex)
+{
+	// Resolve attachments only supported for 2d/2d array textures
+	if (ResolveTexture->GetViewType() == VK_IMAGE_VIEW_TYPE_2D || ResolveTexture->GetViewType() == VK_IMAGE_VIEW_TYPE_2D_ARRAY)
+	{
+		FVulkanTextureViewDesc ViewDesc;
+		ViewDesc.ViewType = ResolveTexture->GetViewType();
+		ViewDesc.AspectFlags = ResolveTexture->GetFullAspectMask();
+		ViewDesc.UEFormat = ResolveTexture->GetDesc().Format;
+		ViewDesc.Format = ResolveTexture->ViewFormat;
+		ViewDesc.FirstMip = MipIndex;
+		ViewDesc.NumMips = 1;
+		ViewDesc.ArraySliceIndex = 0;
+		ViewDesc.NumArraySlices = ResolveTexture->GetNumberOfArrayLevels();
+		ViewDesc.bUseIdentitySwizzle = true;
+		return ResolveTexture->FindOrAddInternalView(ViewDesc);
+	}
+
+	return nullptr;
+}
+
+FVulkanView* FVulkanFramebuffer::GetFragmentDensityAttachmentViewDesc(FVulkanTexture* Texture, uint32 MipIndex)
+{
+	ensure(Texture->GetViewType() == VK_IMAGE_VIEW_TYPE_2D || Texture->GetViewType() == VK_IMAGE_VIEW_TYPE_2D_ARRAY);
+
+	FVulkanTextureViewDesc ViewDesc;
+	ViewDesc.ViewType = Texture->GetViewType();
+	ViewDesc.AspectFlags = Texture->GetFullAspectMask();
+	ViewDesc.UEFormat = Texture->GetDesc().Format;
+	ViewDesc.Format = Texture->ViewFormat;
+	ViewDesc.FirstMip = MipIndex;
+	ViewDesc.NumMips = 1;
+	ViewDesc.ArraySliceIndex = 0;
+	ViewDesc.NumArraySlices = Texture->GetNumberOfArrayLevels();
+	ViewDesc.bUseIdentitySwizzle = true;
+	return Texture->FindOrAddInternalView(ViewDesc);
+}
+
+void FVulkanFramebuffer::Destroy(FVulkanDevice& Device)
+{
+	VulkanRHI::FDeferredDeletionQueue2& Queue = Device.GetDeferredDeletionQueue();
+
+	// will be deleted in reverse order
+	Queue.EnqueueResource(VulkanRHI::FDeferredDeletionQueue2::EType::Framebuffer, Framebuffer);
+	Framebuffer = VK_NULL_HANDLE;
+
+	DEC_DWORD_STAT(STAT_VulkanNumFrameBuffers);
+}
+
+bool FVulkanFramebuffer::Matches(const FRHISetRenderTargetsInfo& InRTInfo) const
+{
+	if (NumColorRenderTargets != InRTInfo.NumColorRenderTargets)
+	{
+		return false;
+	}
+
+	{
+		const FRHIDepthRenderTargetView& B = InRTInfo.DepthStencilRenderTarget;
+		if (B.Texture)
+		{
+			VkImage AImage = DepthStencilRenderTargetImage;
+			VkImage BImage = ResourceCast(B.Texture)->Image;
+			if (AImage != BImage)
+			{
+				return false;
+			}
+		}
+	}
+
+	{
+		const FRHIDepthRenderTargetView& R = InRTInfo.DepthStencilResolveRenderTarget;
+		if (R.Texture)
+		{
+			VkImage AImage = DepthStencilResolveRenderTargetImage;
+			VkImage BImage = ResourceCast(R.Texture)->Image;
+			if (AImage != BImage)
+			{
+				return false;
+			}
+		}
+	}
+
+	{
+		FRHITexture* Texture = InRTInfo.ShadingRateTexture;
+		if (Texture)
+		{
+			VkImage AImage = FragmentDensityImage;
+			VkImage BImage = ResourceCast(Texture)->Image;
+			if (AImage != BImage)
+			{
+				return false;
+			}
+		}
+	}
+
+	int32 AttachementIndex = 0;
+	for (int32 Index = 0; Index < InRTInfo.NumColorRenderTargets; ++Index)
+	{
+		if (InRTInfo.bHasResolveAttachments)
+		{
+			const FRHIRenderTargetView& R = InRTInfo.ColorResolveRenderTarget[Index];
+			if (R.Texture)
+			{
+				VkImage AImage = ColorResolveTargetImages[AttachementIndex];
+				VkImage BImage = ResourceCast(R.Texture)->Image;
+				if (AImage != BImage)
+				{
+					return false;
+				}
+			}
+		}
+
+		const FRHIRenderTargetView& B = InRTInfo.ColorRenderTarget[Index];
+		if (B.Texture)
+		{
+			VkImage AImage = ColorRenderTargetImages[AttachementIndex];
+			VkImage BImage = ResourceCast(B.Texture)->Image;
+			if (AImage != BImage)
+			{
+				return false;
+			}
+			AttachementIndex++;
+		}
+	}
+
+	return true;
+}
+
+
+
+
+
diff --git a/Engine/Source/Runtime/VulkanRHI/Private/VulkanLayout.cpp b/Engine/Source/Runtime/VulkanRHI/Private/VulkanLayout.cpp
new file mode 100644
--- /dev/null
+++ b/Engine/Source/Runtime/VulkanRHI/Private/VulkanLayout.cpp
@@ -0,0 +1,46 @@
+// Copyright Epic Games, Inc. All Rights Reserved.
+
+#include "VulkanLayout.h"
+#include "VulkanBindlessDescriptorManager.h"
+
+FVulkanLayout::FVulkanLayout(FVulkanDevice* InDevice, bool InGfxLayout, bool InUsesBindless)
+	: VulkanRHI::FDeviceChild(InDevice)
+	, DescriptorSetLayout(Device)
+	, PipelineLayout(VK_NULL_HANDLE)
+	, bIsGfxLayout(InGfxLayout)
+	, bUsesBindless(InUsesBindless)
+{
+}
+
+FVulkanLayout::~FVulkanLayout()
+{
+	if (!bUsesBindless && PipelineLayout != VK_NULL_HANDLE)
+	{
+		Device->GetDeferredDeletionQueue().EnqueueResource(VulkanRHI::FDeferredDeletionQueue2::EType::PipelineLayout, PipelineLayout);
+		PipelineLayout = VK_NULL_HANDLE;
+	}
+}
+
+void FVulkanLayout::Compile(FVulkanDescriptorSetLayoutMap& DSetLayoutMap)
+{
+	check(PipelineLayout == VK_NULL_HANDLE);
+
+	DescriptorSetLayout.Compile(DSetLayoutMap);
+
+	if (bUsesBindless)
+	{
+		// Get the bindless manager's layout
+		PipelineLayout = Device->GetBindlessDescriptorManager()->GetPipelineLayout();
+	}
+	else
+	{
+		VkPipelineLayoutCreateInfo PipelineLayoutCreateInfo;
+		ZeroVulkanStruct(PipelineLayoutCreateInfo, VK_STRUCTURE_TYPE_PIPELINE_LAYOUT_CREATE_INFO);
+
+		const TArray<VkDescriptorSetLayout>& LayoutHandles = DescriptorSetLayout.GetHandles();
+		PipelineLayoutCreateInfo.setLayoutCount = LayoutHandles.Num();
+		PipelineLayoutCreateInfo.pSetLayouts = LayoutHandles.GetData();
+		//PipelineLayoutCreateInfo.pushConstantRangeCount = 0;
+		VERIFYVULKANRESULT(VulkanRHI::vkCreatePipelineLayout(Device->GetInstanceHandle(), &PipelineLayoutCreateInfo, VULKAN_CPU_ALLOCATOR, &PipelineLayout));
+	}
+}
diff --git a/Engine/Source/Runtime/VulkanRHI/Private/VulkanLayout.h b/Engine/Source/Runtime/VulkanRHI/Private/VulkanLayout.h
new file mode 100644
--- /dev/null
+++ b/Engine/Source/Runtime/VulkanRHI/Private/VulkanLayout.h
@@ -0,0 +1,70 @@
+// Copyright Epic Games, Inc. All Rights Reserved.
+
+#pragma once
+
+#include "VulkanMemory.h"
+#include "VulkanDescriptorSets.h"
+
+// Layout for a Pipeline, also includes DescriptorSets layout
+class FVulkanLayout : public VulkanRHI::FDeviceChild
+{
+public:
+	FVulkanLayout(FVulkanDevice* InDevice, bool InGfxLayout, bool InUsesBindless);
+	virtual ~FVulkanLayout();
+
+	const FVulkanDescriptorSetsLayout& GetDescriptorSetsLayout() const
+	{
+		return DescriptorSetLayout;
+	}
+
+	VkPipelineLayout GetPipelineLayout() const
+	{
+		return PipelineLayout;
+	}
+
+	bool HasDescriptors() const
+	{
+		return DescriptorSetLayout.GetLayouts().Num() > 0;
+	}
+
+	uint32 GetDescriptorSetLayoutHash() const
+	{
+		return DescriptorSetLayout.GetHash();
+	}
+
+	bool IsGfxLayout() const
+	{
+		return bIsGfxLayout;
+	}
+
+protected:
+	FVulkanDescriptorSetsLayout	DescriptorSetLayout;
+	VkPipelineLayout			PipelineLayout;
+
+	const bool bIsGfxLayout;
+	const bool bUsesBindless;
+
+	template <bool bIsCompute>
+	void FinalizeBindings(const FUniformBufferGatherInfo& UBGatherInfo)
+	{
+		// Setting descriptor is only allowed prior to compiling the layout
+		check(DescriptorSetLayout.GetHandles().Num() == 0);
+
+		DescriptorSetLayout.FinalizeBindings<bIsCompute>(UBGatherInfo);
+	}
+
+	void ProcessBindingsForStage(VkShaderStageFlagBits StageFlags, ShaderStage::EStage DescSet, const FVulkanShaderHeader& CodeHeader, FUniformBufferGatherInfo& OutUBGatherInfo) const
+	{
+		// Setting descriptor is only allowed prior to compiling the layout
+		check(DescriptorSetLayout.GetHandles().Num() == 0);
+
+		DescriptorSetLayout.ProcessBindingsForStage(StageFlags, DescSet, CodeHeader, OutUBGatherInfo);
+	}
+
+	void Compile(FVulkanDescriptorSetLayoutMap& DSetLayoutMap);
+
+	friend class FVulkanComputePipeline;
+	friend class FVulkanGfxPipeline;
+	friend class FVulkanPipelineStateCacheManager;
+	friend class FVulkanRayTracingPipelineState;
+};
diff --git a/Engine/Source/Runtime/VulkanRHI/Private/VulkanMemory.cpp b/Engine/Source/Runtime/VulkanRHI/Private/VulkanMemory.cpp
--- a/Engine/Source/Runtime/VulkanRHI/Private/VulkanMemory.cpp
+++ b/Engine/Source/Runtime/VulkanRHI/Private/VulkanMemory.cpp
@@ -11,6 +11,7 @@
 #include "VulkanContext.h"
 #include "VulkanLLM.h"
 #include "VulkanDescriptorSets.h"
+#include "VulkanBindlessDescriptorManager.h"
 #include "Containers/SortedMap.h"
 
 #if PLATFORM_WINDOWS
diff --git a/Engine/Source/Runtime/VulkanRHI/Private/VulkanPendingState.cpp b/Engine/Source/Runtime/VulkanRHI/Private/VulkanPendingState.cpp
--- a/Engine/Source/Runtime/VulkanRHI/Private/VulkanPendingState.cpp
+++ b/Engine/Source/Runtime/VulkanRHI/Private/VulkanPendingState.cpp
@@ -6,6 +6,7 @@
 
 #include "VulkanRHIPrivate.h"
 #include "VulkanCommandWrappers.h"
+#include "VulkanFramebuffer.h"
 #include "VulkanPendingState.h"
 #include "VulkanPipeline.h"
 #include "VulkanContext.h"
diff --git a/Engine/Source/Runtime/VulkanRHI/Private/VulkanPipeline.cpp b/Engine/Source/Runtime/VulkanRHI/Private/VulkanPipeline.cpp
--- a/Engine/Source/Runtime/VulkanRHI/Private/VulkanPipeline.cpp
+++ b/Engine/Source/Runtime/VulkanRHI/Private/VulkanPipeline.cpp
@@ -18,6 +18,8 @@
 #include "VulkanLLM.h"
 #include "Misc/ScopeRWLock.h"
 #include "VulkanChunkedPipelineCache.h"
+#include "VulkanRenderTargetLayout.h"
+#include "VulkanRenderpass.h"
 
 #define LRU_DEBUG 0
 #if !UE_BUILD_SHIPPING
diff --git a/Engine/Source/Runtime/VulkanRHI/Private/VulkanPipelineState.cpp b/Engine/Source/Runtime/VulkanRHI/Private/VulkanPipelineState.cpp
--- a/Engine/Source/Runtime/VulkanRHI/Private/VulkanPipelineState.cpp
+++ b/Engine/Source/Runtime/VulkanRHI/Private/VulkanPipelineState.cpp
@@ -12,6 +12,7 @@
 #include "VulkanPendingState.h"
 #include "VulkanPipeline.h"
 #include "VulkanLLM.h"
+#include "VulkanBindlessDescriptorManager.h"
 #include "RHICoreShader.h"
 #include "GlobalRenderResources.h"  // For GBlackTexture
 
@@ -622,113 +623,6 @@ template bool FVulkanComputePipelineDescriptorState::InternalUpdateDescriptorSet
 template bool FVulkanComputePipelineDescriptorState::InternalUpdateDescriptorSets<false>(FVulkanCommandListContext& Context);
 
 
-void FVulkanDescriptorSetWriter::CheckAllWritten()
-{
-#if VULKAN_VALIDATE_DESCRIPTORS_WRITTEN
-	auto GetVkDescriptorTypeString = [](VkDescriptorType Type)
-	{
-		switch (Type)
-		{
-			// + 19 to skip "VK_DESCRIPTOR_TYPE_"
-	#define VKSWITCHCASE(x)	case x: return FString(&TEXT(#x)[19]);
-			VKSWITCHCASE(VK_DESCRIPTOR_TYPE_SAMPLER)
-				VKSWITCHCASE(VK_DESCRIPTOR_TYPE_COMBINED_IMAGE_SAMPLER)
-				VKSWITCHCASE(VK_DESCRIPTOR_TYPE_SAMPLED_IMAGE)
-				VKSWITCHCASE(VK_DESCRIPTOR_TYPE_STORAGE_IMAGE)
-				VKSWITCHCASE(VK_DESCRIPTOR_TYPE_UNIFORM_TEXEL_BUFFER)
-				VKSWITCHCASE(VK_DESCRIPTOR_TYPE_STORAGE_TEXEL_BUFFER)
-				VKSWITCHCASE(VK_DESCRIPTOR_TYPE_UNIFORM_BUFFER)
-				VKSWITCHCASE(VK_DESCRIPTOR_TYPE_STORAGE_BUFFER)
-				VKSWITCHCASE(VK_DESCRIPTOR_TYPE_UNIFORM_BUFFER_DYNAMIC)
-				VKSWITCHCASE(VK_DESCRIPTOR_TYPE_STORAGE_BUFFER_DYNAMIC)
-				VKSWITCHCASE(VK_DESCRIPTOR_TYPE_INPUT_ATTACHMENT)
-	#undef VKSWITCHCASE
-		default:
-			break;
-		}
-
-		return FString::Printf(TEXT("Unknown VkDescriptorType %d"), (int32)Type);
-	};
-
-	const uint32 Writes = NumWrites;
-	if (Writes == 0)
-		return;
-
-	bool bFail = false;
-	if(Writes <= 32) //early out for the most common case.
-	{
-		bFail = WrittenMask[0] != ((1llu << Writes)-1);
-	}
-	else
-	{
-		const int32 Last = int32(WrittenMask.Num()-1);
-		for(int32 i = 0; !bFail && i < Last; ++i)
-		{
-			uint64 Mask = WrittenMask[i];
-			bFail = bFail || Mask != 0xffffffff; 
-		}
-
-		const uint32 TailCount = Writes - (Last * 32);
-		check(TailCount != 0);
-		const uint32 TailMask = (1llu << TailCount)-1;
-		bFail = bFail || TailMask != WrittenMask[Last];
-	}
-
-	if(bFail)
-	{
-		FString Descriptors;
-		for (uint32 i = 0; i < Writes; ++i)
-		{
-			uint32 Index = i / 32;
-			uint32 Mask = i % 32;
-			if(0 == (WrittenMask[Index] & (1llu << Mask)))
-			{
-				FString TypeString = GetVkDescriptorTypeString(WriteDescriptors[i].descriptorType);
-				Descriptors += FString::Printf(TEXT("\t\tDescriptorWrite %d/%d Was not written(Type %s)\n"), i, NumWrites, *TypeString);
-			}
-		}
-		UE_LOG(LogVulkanRHI, Warning, TEXT("Not All descriptors where filled out. this can/will cause a driver crash\n%s\n"), *Descriptors);
-		ensureMsgf(false, TEXT("Not All descriptors where filled out. this can/will cause a driver crash\n%s\n"), *Descriptors);
-	}
-#endif
-}
-
-void FVulkanDescriptorSetWriter::Reset()
-{
-	bHasVolatileResources = false;
-
-#if VULKAN_VALIDATE_DESCRIPTORS_WRITTEN
-	WrittenMask = BaseWrittenMask;
-#endif
-}
-void FVulkanDescriptorSetWriter::SetWritten(uint32 DescriptorIndex)
-{
-#if VULKAN_VALIDATE_DESCRIPTORS_WRITTEN
-	uint32 Index = DescriptorIndex / 32;
-	uint32 Mask = DescriptorIndex % 32;
-	WrittenMask[Index] |= (1<<Mask);
-#endif
-}
-void FVulkanDescriptorSetWriter::SetWrittenBase(uint32 DescriptorIndex)
-{
-#if VULKAN_VALIDATE_DESCRIPTORS_WRITTEN	
-	uint32 Index = DescriptorIndex / 32;
-	uint32 Mask = DescriptorIndex % 32;
-	BaseWrittenMask[Index] |= (1<<Mask);
-#endif
-}
-
-void FVulkanDescriptorSetWriter::InitWrittenMasks(uint32 NumDescriptorWrites)
-{
-#if VULKAN_VALIDATE_DESCRIPTORS_WRITTEN	
-	uint32 Size = (NumDescriptorWrites + 31) / 32;
-	WrittenMask.Empty(Size);
-	WrittenMask.SetNumZeroed(Size);
-	BaseWrittenMask.Empty(Size);
-	BaseWrittenMask.SetNumZeroed(Size);
-#endif
-}
-
 void FVulkanCommonPipelineDescriptorState::SetSRV(bool bCompute, uint8 DescriptorSet, uint32 BindingIndex, FVulkanShaderResourceView* SRV)
 {
 	check(!bUseBindless);
diff --git a/Engine/Source/Runtime/VulkanRHI/Private/VulkanPipelineState.h b/Engine/Source/Runtime/VulkanRHI/Private/VulkanPipelineState.h
--- a/Engine/Source/Runtime/VulkanRHI/Private/VulkanPipelineState.h
+++ b/Engine/Source/Runtime/VulkanRHI/Private/VulkanPipelineState.h
@@ -10,6 +10,7 @@
 #include "VulkanMemory.h"
 #include "VulkanCommandBuffer.h"
 #include "VulkanDescriptorSets.h"
+#include "VulkanLayout.h"
 #include "VulkanPipeline.h"
 #include "VulkanRHIPrivate.h"
 #include "Containers/ArrayView.h"
