Reorganizing a number of Vulkan types to be in their own file(s).

diff --git a/Engine/Source/Runtime/VulkanRHI/Private/VulkanRHI.cpp b/Engine/Source/Runtime/VulkanRHI/Private/VulkanRHI.cpp
--- a/Engine/Source/Runtime/VulkanRHI/Private/VulkanRHI.cpp
+++ b/Engine/Source/Runtime/VulkanRHI/Private/VulkanRHI.cpp
@@ -30,6 +30,7 @@
 #include "VulkanExtensions.h"
 #include "VulkanRayTracing.h"
 #include "VulkanChunkedPipelineCache.h"
+#include "VulkanBindlessDescriptorManager.h"
 #if PLATFORM_ANDROID
 #include "Android/AndroidPlatformMisc.h"
 #endif
@@ -1770,245 +1771,6 @@ FTextureRHIRef FVulkanDynamicRHI::RHICreateAliasedTexture(FTextureRHIRef& Source
 	return new FVulkanTexture(*Device, Desc, SourceTextureRHI);
 }
 
-FVulkanDescriptorSetsLayout::FVulkanDescriptorSetsLayout(FVulkanDevice* InDevice) :
-	Device(InDevice)
-{
-}
-
-FVulkanDescriptorSetsLayout::~FVulkanDescriptorSetsLayout()
-{
-	// Handles are owned by FVulkanPipelineStateCacheManager
-	LayoutHandles.Reset(0);
-}
-
-// Increments a value and asserts on overflow.
-// FSetInfo uses narrow integer types for descriptor counts,
-// which may feasibly overflow one day (for example if we add bindless resources).
-template <typename T> 
-static void IncrementChecked(T& Value)
-{
-	check(Value < TNumericLimits<T>::Max());
-	++Value;
-}
-
-void FVulkanDescriptorSetsLayoutInfo::AddDescriptor(int32 DescriptorSetIndex, const VkDescriptorSetLayoutBinding& Descriptor)
-{
-	// Increment type usage
-	if (LayoutTypes.Contains(Descriptor.descriptorType))
-	{
-		LayoutTypes[Descriptor.descriptorType]++;
-	}
-	else
-	{
-		LayoutTypes.Add(Descriptor.descriptorType, 1);
-	}
-
-	if (DescriptorSetIndex >= SetLayouts.Num())
-	{
-		SetLayouts.SetNum(DescriptorSetIndex + 1, EAllowShrinking::No);
-	}
-
-	FSetLayout& DescSetLayout = SetLayouts[DescriptorSetIndex];
-
-	VkDescriptorSetLayoutBinding* Binding = new(DescSetLayout.LayoutBindings) VkDescriptorSetLayoutBinding;
-	*Binding = Descriptor;
-
-	const FStageInfo& SetInfo = StageInfos[DescriptorSetIndex];
-	check(SetInfo.Types[Descriptor.binding] == Descriptor.descriptorType);
-	switch (Descriptor.descriptorType)
-	{
-	case VK_DESCRIPTOR_TYPE_SAMPLER:
-	case VK_DESCRIPTOR_TYPE_COMBINED_IMAGE_SAMPLER:
-	case VK_DESCRIPTOR_TYPE_SAMPLED_IMAGE:
-	case VK_DESCRIPTOR_TYPE_STORAGE_IMAGE:
-	case VK_DESCRIPTOR_TYPE_INPUT_ATTACHMENT:
-		IncrementChecked(StageInfos[DescriptorSetIndex].NumImageInfos);
-		break;
-	case VK_DESCRIPTOR_TYPE_UNIFORM_BUFFER:
-	case VK_DESCRIPTOR_TYPE_STORAGE_BUFFER:
-	case VK_DESCRIPTOR_TYPE_UNIFORM_BUFFER_DYNAMIC:
-		IncrementChecked(StageInfos[DescriptorSetIndex].NumBufferInfos);
-		break;
-	case VK_DESCRIPTOR_TYPE_ACCELERATION_STRUCTURE_KHR:
-		IncrementChecked(StageInfos[DescriptorSetIndex].NumAccelerationStructures);
-		break;
-	case VK_DESCRIPTOR_TYPE_STORAGE_TEXEL_BUFFER:
-	case VK_DESCRIPTOR_TYPE_UNIFORM_TEXEL_BUFFER:
-		break;
-	default:
-		checkf(0, TEXT("Unsupported descriptor type %d"), (int32)Descriptor.descriptorType);
-		break;
-	}
-}
-
-void FVulkanDescriptorSetsLayoutInfo::GenerateHash(const TArrayView<FRHISamplerState*>& InImmutableSamplers, VkPipelineBindPoint InBindPoint)
-{
-	const int32 LayoutCount = SetLayouts.Num();
-	Hash = FCrc::MemCrc32(&TypesUsageID, sizeof(uint32), LayoutCount);
-
-	for (int32 layoutIndex = 0; layoutIndex < LayoutCount; ++layoutIndex)
-	{
-		SetLayouts[layoutIndex].GenerateHash();
-		Hash = FCrc::MemCrc32(&SetLayouts[layoutIndex].Hash, sizeof(uint32), Hash);
-	}
-
-	const uint32 NumStages = GetNumStagesForBindPoint(InBindPoint);
-	for (uint32 RemapingIndex = 0; RemapingIndex < NumStages; ++RemapingIndex)
-	{
-		const FStageInfo& StageInfo = StageInfos[RemapingIndex];
-
-		Hash = FCrc::TypeCrc32(StageInfo.PackedGlobalsSize, Hash);
-		Hash = FCrc::TypeCrc32(StageInfo.NumBoundUniformBuffers, Hash);
-		Hash = FCrc::TypeCrc32(StageInfo.NumImageInfos, Hash);
-		Hash = FCrc::TypeCrc32(StageInfo.NumBufferInfos, Hash);
-		Hash = FCrc::TypeCrc32(StageInfo.NumAccelerationStructures, Hash);
-
-		const TArray<VkDescriptorType>& Types = StageInfo.Types;
-		Hash = FCrc::MemCrc32(Types.GetData(), sizeof(VkDescriptorType) * Types.Num(), Hash);
-	}
-
-	// It would be better to store this when the object is created, but it's not available at that time, so we'll do it here.
-	BindPoint = InBindPoint;
-
-	// Include the bind point in the hash, because we can have graphics and compute PSOs with the same descriptor info, and we don't want them to collide.
-	Hash = FCrc::MemCrc32(&BindPoint, sizeof(BindPoint), Hash);
-}
-
-static FCriticalSection GTypesUsageCS;
-void FVulkanDescriptorSetsLayoutInfo::CompileTypesUsageID()
-{
-	FScopeLock ScopeLock(&GTypesUsageCS);
-
-	static TMap<uint32, uint32> GTypesUsageHashMap;
-	static uint32 GUniqueID = 1;
-
-	LayoutTypes.KeySort([](VkDescriptorType A, VkDescriptorType B)
-	{
-		return static_cast<uint32>(A) < static_cast<uint32>(B);
-	});
-
-	uint32 TypesUsageHash = 0;
-	for (const auto& Elem : LayoutTypes)
-	{
-		TypesUsageHash = FCrc::MemCrc32(&Elem.Value, sizeof(uint32), TypesUsageHash);
-	}
-
-	uint32* UniqueID = GTypesUsageHashMap.Find(TypesUsageHash);
-	if (UniqueID == nullptr)
-	{
-		TypesUsageID = GTypesUsageHashMap.Add(TypesUsageHash, GUniqueID++);
-	}
-	else
-	{
-		TypesUsageID = *UniqueID;
-	}
-}
-
-void FVulkanDescriptorSetsLayout::Compile(FVulkanDescriptorSetLayoutMap& DSetLayoutMap)
-{
-	check(LayoutHandles.Num() == 0);
-
-	// Check if we obey limits
-	const VkPhysicalDeviceLimits& Limits = Device->GetLimits();
-	
-	// Check for maxDescriptorSetSamplers
-	check(		LayoutTypes[VK_DESCRIPTOR_TYPE_SAMPLER]
-			+	LayoutTypes[VK_DESCRIPTOR_TYPE_COMBINED_IMAGE_SAMPLER]
-			<=	Limits.maxDescriptorSetSamplers);
-
-	// Check for maxDescriptorSetUniformBuffers
-	check(		LayoutTypes[VK_DESCRIPTOR_TYPE_UNIFORM_BUFFER]
-			+	LayoutTypes[VK_DESCRIPTOR_TYPE_UNIFORM_BUFFER_DYNAMIC]
-			<=	Limits.maxDescriptorSetUniformBuffers);
-
-	// Check for maxDescriptorSetUniformBuffersDynamic
-	check(Device->GetVendorId() == EGpuVendorId::Amd ||
-				LayoutTypes[VK_DESCRIPTOR_TYPE_UNIFORM_BUFFER_DYNAMIC]
-			<=	Limits.maxDescriptorSetUniformBuffersDynamic);
-
-	// Check for maxDescriptorSetStorageBuffers
-	check(		LayoutTypes[VK_DESCRIPTOR_TYPE_STORAGE_BUFFER]
-			+	LayoutTypes[VK_DESCRIPTOR_TYPE_STORAGE_BUFFER_DYNAMIC]
-			<=	Limits.maxDescriptorSetStorageBuffers);
-
-	// Check for maxDescriptorSetStorageBuffersDynamic
-	check(		LayoutTypes[VK_DESCRIPTOR_TYPE_STORAGE_BUFFER_DYNAMIC]
-			<=	Limits.maxDescriptorSetStorageBuffersDynamic);
-
-	// Check for maxDescriptorSetSampledImages
-	check(		LayoutTypes[VK_DESCRIPTOR_TYPE_COMBINED_IMAGE_SAMPLER]
-			+	LayoutTypes[VK_DESCRIPTOR_TYPE_SAMPLED_IMAGE]
-			+	LayoutTypes[VK_DESCRIPTOR_TYPE_UNIFORM_TEXEL_BUFFER]
-			<=	Limits.maxDescriptorSetSampledImages);
-
-	// Check for maxDescriptorSetStorageImages
-	check(		LayoutTypes[VK_DESCRIPTOR_TYPE_STORAGE_IMAGE]
-			+	LayoutTypes[VK_DESCRIPTOR_TYPE_STORAGE_TEXEL_BUFFER]
-			<=	Limits.maxDescriptorSetStorageImages);
-
-	check(LayoutTypes[VK_DESCRIPTOR_TYPE_INPUT_ATTACHMENT] <= Limits.maxDescriptorSetInputAttachments);
-
-	if (GRHISupportsRayTracing)
-	{
-		check(LayoutTypes[VK_DESCRIPTOR_TYPE_ACCELERATION_STRUCTURE_KHR] < Device->GetOptionalExtensionProperties().AccelerationStructureProps.maxDescriptorSetAccelerationStructures);
-	}
-	
-	LayoutHandles.Empty(SetLayouts.Num());
-
-	if (UseVulkanDescriptorCache())
-	{
-		LayoutHandleIds.Empty(SetLayouts.Num());
-	}
-				
-	for (FSetLayout& Layout : SetLayouts)
-	{
-		VkDescriptorSetLayout* LayoutHandle = new(LayoutHandles) VkDescriptorSetLayout;
-
-		uint32* LayoutHandleId = nullptr;
-		if (UseVulkanDescriptorCache())
-		{
-			LayoutHandleId = new(LayoutHandleIds) uint32;
-		}
-			
-		if (FVulkanDescriptorSetLayoutEntry* Found = DSetLayoutMap.Find(Layout))
-		{
-			*LayoutHandle = Found->Handle;
-			if (LayoutHandleId)
-			{
-				*LayoutHandleId = Found->HandleId;
-			}
-			continue;
-		}
-
-		VkDescriptorSetLayoutCreateInfo DescriptorLayoutInfo;
-		ZeroVulkanStruct(DescriptorLayoutInfo, VK_STRUCTURE_TYPE_DESCRIPTOR_SET_LAYOUT_CREATE_INFO);
-		DescriptorLayoutInfo.bindingCount = Layout.LayoutBindings.Num();
-		DescriptorLayoutInfo.pBindings = Layout.LayoutBindings.GetData();
-
-		VERIFYVULKANRESULT(VulkanRHI::vkCreateDescriptorSetLayout(Device->GetInstanceHandle(), &DescriptorLayoutInfo, VULKAN_CPU_ALLOCATOR, LayoutHandle));
-
-		if (LayoutHandleId)
-		{
-			*LayoutHandleId = ++GVulkanDSetLayoutHandleIdCounter;
-		}
-
-		FVulkanDescriptorSetLayoutEntry DescriptorSetLayoutEntry;
-		DescriptorSetLayoutEntry.Handle = *LayoutHandle;
-		DescriptorSetLayoutEntry.HandleId = LayoutHandleId ? *LayoutHandleId : 0;
-				
-		DSetLayoutMap.Add(Layout, DescriptorSetLayoutEntry);
-	}
-
-	if (TypesUsageID == ~0)
-	{
-		CompileTypesUsageID();
-	}
-
-	ZeroVulkanStruct(DescriptorSetAllocateInfo, VK_STRUCTURE_TYPE_DESCRIPTOR_SET_ALLOCATE_INFO);
-	DescriptorSetAllocateInfo.descriptorSetCount = LayoutHandles.Num();
-	DescriptorSetAllocateInfo.pSetLayouts = LayoutHandles.GetData();
-}
-
 FVulkanRenderPass::FVulkanRenderPass(FVulkanDevice& InDevice, const FVulkanRenderTargetLayout& InRTLayout) :
 	Layout(InRTLayout),
 	RenderPass(VK_NULL_HANDLE),
diff --git a/Engine/Source/Runtime/VulkanRHI/Private/VulkanRHIPrivate.h b/Engine/Source/Runtime/VulkanRHI/Private/VulkanRHIPrivate.h
--- a/Engine/Source/Runtime/VulkanRHI/Private/VulkanRHIPrivate.h
+++ b/Engine/Source/Runtime/VulkanRHI/Private/VulkanRHIPrivate.h
@@ -1,9 +1,5 @@
 // Copyright Epic Games, Inc. All Rights Reserved..
 
-/*=============================================================================
-	VulkanRHIPrivate.h: Private Vulkan RHI definitions.
-=============================================================================*/
-
 #pragma once
 
 // Dependencies
@@ -123,248 +119,6 @@ static constexpr int32 GetNumStagesForBindPoint(VkPipelineBindPoint BindPoint)
 	return 0;
 }
 
-class FVulkanRenderTargetLayout
-{
-public:
-	FVulkanRenderTargetLayout(const FGraphicsPipelineStateInitializer& Initializer);
-	FVulkanRenderTargetLayout(FVulkanDevice& InDevice, const FRHISetRenderTargetsInfo& RTInfo);
-	FVulkanRenderTargetLayout(FVulkanDevice& InDevice, const FRHIRenderPassInfo& RPInfo, VkImageLayout CurrentDepthLayout, VkImageLayout CurrentStencilLayout);
-
-	inline uint32 GetRenderPassCompatibleHash() const
-	{
-		check(bCalculatedHash);
-		return RenderPassCompatibleHash;
-	}
-	inline uint32 GetRenderPassFullHash() const
-	{
-		check(bCalculatedHash);
-		return RenderPassFullHash;
-	}
-	inline const VkOffset2D& GetOffset2D() const { return Offset.Offset2D; }
-	inline const VkOffset3D& GetOffset3D() const { return Offset.Offset3D; }
-	inline const VkExtent2D& GetExtent2D() const { return Extent.Extent2D; }
-	inline const VkExtent3D& GetExtent3D() const { return Extent.Extent3D; }
-	inline const VkAttachmentDescription* GetAttachmentDescriptions() const { return Desc; }
-	inline uint32 GetNumColorAttachments() const { return NumColorAttachments; }
-	inline bool GetHasDepthStencil() const { return bHasDepthStencil != 0; }
-	inline bool GetHasResolveAttachments() const { return bHasResolveAttachments != 0; }
-	inline bool GetHasDepthStencilResolve() const { return bHasDepthStencilResolve != 0; }
-	inline bool GetHasFragmentDensityAttachment() const { return bHasFragmentDensityAttachment != 0; }
-	inline uint32 GetNumAttachmentDescriptions() const { return NumAttachmentDescriptions; }
-	inline uint32 GetNumSamples() const { return NumSamples; }
-	inline uint32 GetNumUsedClearValues() const { return NumUsedClearValues; }
-	inline bool GetIsMultiView() const { return MultiViewCount != 0; }
-	inline uint32 GetMultiViewCount() const { return MultiViewCount; }
-
-
-	inline const VkAttachmentReference* GetColorAttachmentReferences() const { return NumColorAttachments > 0 ? ColorReferences : nullptr; }
-	inline const VkAttachmentReference* GetResolveAttachmentReferences() const { return bHasResolveAttachments ? ResolveReferences : nullptr; }
-	inline const VkAttachmentReference* GetDepthAttachmentReference() const { return bHasDepthStencil ? &DepthReference : nullptr; }
-	inline const VkAttachmentReferenceStencilLayout* GetStencilAttachmentReference() const { return bHasDepthStencil ? &StencilReference : nullptr; }
-	inline const VkAttachmentReference* GetDepthStencilResolveAttachmentReference() const { return bHasDepthStencilResolve ? &DepthStencilResolveReference : nullptr; }
-	inline const VkAttachmentReference* GetFragmentDensityAttachmentReference() const { return bHasFragmentDensityAttachment ? &FragmentDensityReference : nullptr; }
-
-	inline const VkAttachmentDescriptionStencilLayout* GetStencilDesc() const { return bHasDepthStencil ? &StencilDesc : nullptr; }
-
-	inline const ESubpassHint GetSubpassHint() const { return SubpassHint; }
-
-protected:
-	VkImageLayout GetVRSImageLayout() const;
-
-protected:
-	VkAttachmentReference ColorReferences[MaxSimultaneousRenderTargets];
-	VkAttachmentReference DepthReference;
-	VkAttachmentReferenceStencilLayout StencilReference;
-	VkAttachmentReference FragmentDensityReference;
-	VkAttachmentReference ResolveReferences[MaxSimultaneousRenderTargets];
-	VkAttachmentReference DepthStencilResolveReference;
-
-	// Depth goes in the "+1" slot, Depth resolve goes in the "+2 slot", and the Shading Rate texture goes in the "+3" slot.
-	VkAttachmentDescription Desc[MaxSimultaneousRenderTargets * 2 + 3];
-	VkAttachmentDescriptionStencilLayout StencilDesc;
-
-	uint8 NumAttachmentDescriptions;
-	uint8 NumColorAttachments;
-	uint8 NumInputAttachments = 0;
-	uint8 bHasDepthStencil;
-	uint8 bHasResolveAttachments;
-	uint8 bHasDepthStencilResolve;
-	uint8 bHasFragmentDensityAttachment;
-	uint8 NumSamples;
-	uint8 NumUsedClearValues;
-	ESubpassHint SubpassHint = ESubpassHint::None;
-	uint8 MultiViewCount;
-
-	// Hash for a compatible RenderPass
-	uint32 RenderPassCompatibleHash = 0;
-	// Hash for the render pass including the load/store operations
-	uint32 RenderPassFullHash = 0;
-
-	union
-	{
-		VkOffset3D Offset3D;
-		VkOffset2D Offset2D;
-	} Offset;
-
-	union
-	{
-		VkExtent3D	Extent3D;
-		VkExtent2D	Extent2D;
-	} Extent;
-
-	inline void ResetAttachments()
-	{
-		FMemory::Memzero(ColorReferences);
-		FMemory::Memzero(DepthReference);
-		FMemory::Memzero(FragmentDensityReference);
-		FMemory::Memzero(ResolveReferences);
-		FMemory::Memzero(DepthStencilResolveReference);
-		FMemory::Memzero(Desc);
-		FMemory::Memzero(Offset);
-		FMemory::Memzero(Extent);
-
-		ZeroVulkanStruct(StencilReference, VK_STRUCTURE_TYPE_ATTACHMENT_REFERENCE_STENCIL_LAYOUT);
-		ZeroVulkanStruct(StencilDesc, VK_STRUCTURE_TYPE_ATTACHMENT_DESCRIPTION_STENCIL_LAYOUT);
-	}
-
-	FVulkanRenderTargetLayout()
-	{
-		NumAttachmentDescriptions = 0;
-		NumColorAttachments = 0;
-		bHasDepthStencil = 0;
-		bHasResolveAttachments = 0;
-		bHasDepthStencilResolve = 0;
-		bHasFragmentDensityAttachment = 0;
-		NumSamples = 0;
-		NumUsedClearValues = 0;
-		MultiViewCount = 0;
-
-		ResetAttachments();
-	}
-
-	bool bCalculatedHash = false;
-
-	friend class FVulkanPipelineStateCacheManager;
-	friend struct FGfxPipelineDesc;
-};
-
-class FVulkanFramebuffer
-{
-public:
-	FVulkanFramebuffer(FVulkanDevice& Device, const FRHISetRenderTargetsInfo& InRTInfo, const FVulkanRenderTargetLayout& RTLayout, const FVulkanRenderPass& RenderPass);
-	~FVulkanFramebuffer();
-
-	bool Matches(const FRHISetRenderTargetsInfo& RTInfo) const;
-
-	uint32 GetNumColorAttachments() const
-	{
-		return NumColorAttachments;
-	}
-
-	void Destroy(FVulkanDevice& Device);
-
-	VkFramebuffer GetHandle()
-	{
-		return Framebuffer;
-	}
-
-	const FVulkanView::FTextureView& GetPartialDepthTextureView() const
-	{
-		check(PartialDepthTextureView);
-		return PartialDepthTextureView->GetTextureView();
-	}
-
-	TArray<FVulkanView const*> AttachmentTextureViews;
-
-	// Copy from the Depth render target partial view
-	FVulkanView const* PartialDepthTextureView = nullptr;
-
-	bool ContainsRenderTarget(FRHITexture* Texture) const
-	{
-		ensure(Texture);
-		FVulkanTexture* VulkanTexture = ResourceCast(Texture);
-		return ContainsRenderTarget(VulkanTexture->Image);
-	}
-
-	bool ContainsRenderTarget(VkImage Image) const
-	{
-		ensure(Image != VK_NULL_HANDLE);
-		for (uint32 Index = 0; Index < NumColorAttachments; ++Index)
-		{
-			if (ColorRenderTargetImages[Index] == Image)
-			{
-				return true;
-			}
-		}
-
-		return (DepthStencilRenderTargetImage == Image);
-	}
-
-	VkRect2D GetRenderArea() const
-	{
-		return RenderArea;
-	}
-
-	// Expose the view creation logic so it can be reused for dynamic rendering
-	static FVulkanView* GetColorRenderTargetViewDesc(FVulkanTexture* Texture, uint32 MipIndex, int32 ArraySliceIndex, uint32 MultiViewCount, uint32& InOutNumLayers);
-	static FVulkanView* GetColorResolveTargetViewDesc(FVulkanTexture* ResolveTexture, uint32 MipIndex, int32 ArraySliceIndex);
-	static FVulkanView* GetDepthStencilTargetViewDesc(FVulkanTexture* Texture, uint32 NumColorAttachments, uint32 MipIndex, uint32& InOutNumLayers);
-	static FVulkanView* GetDepthStencilResolveTargetViewDesc(FVulkanTexture* ResolveTexture, uint32 MipIndex);
-	static FVulkanView* GetFragmentDensityAttachmentViewDesc(FVulkanTexture* Texture, uint32 MipIndex);
-
-private:
-	VkFramebuffer Framebuffer;
-	VkRect2D RenderArea;
-
-	// Unadjusted number of color render targets as in FRHISetRenderTargetsInfo 
-	uint32 NumColorRenderTargets;
-
-	// Save image off for comparison, in case it gets aliased.
-	uint32 NumColorAttachments;
-	VkImage ColorRenderTargetImages[MaxSimultaneousRenderTargets];
-	VkImage ColorResolveTargetImages[MaxSimultaneousRenderTargets];
-	VkImage DepthStencilRenderTargetImage;
-	VkImage DepthStencilResolveRenderTargetImage;
-	VkImage FragmentDensityImage;
-
-	// Predefined set of barriers, when executes ensuring all writes are finished
-	TArray<VkImageMemoryBarrier> WriteBarriers;
-
-	friend class FVulkanCommandListContext;
-};
-
-class FVulkanRenderPass
-{
-public:
-	inline const FVulkanRenderTargetLayout& GetLayout() const
-	{
-		return Layout;
-	}
-
-	inline VkRenderPass GetHandle() const
-	{
-		return RenderPass;
-	}
-
-	inline uint32 GetNumUsedClearValues() const
-	{
-		return NumUsedClearValues;
-	}
-
-private:
-	friend class FVulkanRenderPassManager;
-	friend class FVulkanPipelineStateCacheManager;
-
-	FVulkanRenderPass(FVulkanDevice& Device, const FVulkanRenderTargetLayout& RTLayout);
-	~FVulkanRenderPass();
-
-private:
-	FVulkanRenderTargetLayout	Layout;
-	VkRenderPass				RenderPass;
-	uint32						NumUsedClearValues;
-	FVulkanDevice&				Device;
-};
-
 union UNvidiaDriverVersion
 {
 	struct
diff --git a/Engine/Source/Runtime/VulkanRHI/Private/VulkanRenderTarget.cpp b/Engine/Source/Runtime/VulkanRHI/Private/VulkanRenderTarget.cpp
--- a/Engine/Source/Runtime/VulkanRHI/Private/VulkanRenderTarget.cpp
+++ b/Engine/Source/Runtime/VulkanRHI/Private/VulkanRenderTarget.cpp
@@ -646,718 +646,3 @@ void FVulkanCommandListContext::RHINextSubpass()
 		VulkanRHI::vkCmdNextSubpass(CommandBufferHandle, VK_SUBPASS_CONTENTS_INLINE);
 	}
 }
-
-// Need a separate struct so we can memzero/remove dependencies on reference counts
-struct FRenderPassCompatibleHashableStruct
-{
-	FRenderPassCompatibleHashableStruct()
-	{
-		FMemory::Memzero(*this);
-	}
-
-	uint8							NumAttachments;
-	uint8							MultiViewCount;
-	uint8							NumSamples;
-	uint8							SubpassHint;
-	// +1 for Depth, +1 for Stencil, +1 for Fragment Density
-	VkFormat						Formats[MaxSimultaneousRenderTargets + 3];
-	uint16							AttachmentsToResolve;
-};
-
-// Need a separate struct so we can memzero/remove dependencies on reference counts
-struct FRenderPassFullHashableStruct
-{
-	FRenderPassFullHashableStruct()
-	{
-		FMemory::Memzero(*this);
-	}
-
-	// +1 for Depth, +1 for Stencil, +1 for Fragment Density
-	TEnumAsByte<VkAttachmentLoadOp>		LoadOps[MaxSimultaneousRenderTargets + 3];
-	TEnumAsByte<VkAttachmentStoreOp>	StoreOps[MaxSimultaneousRenderTargets + 3];
-	// If the initial != final we need to add FinalLayout and potentially RefLayout
-	VkImageLayout						InitialLayout[MaxSimultaneousRenderTargets + 3];
-	//VkImageLayout						FinalLayout[MaxSimultaneousRenderTargets + 3];
-	//VkImageLayout						RefLayout[MaxSimultaneousRenderTargets + 3];
-};
-
-VkImageLayout FVulkanRenderTargetLayout::GetVRSImageLayout() const
-{
-	if (ValidateShadingRateDataType())
-	{
-		if (GRHIVariableRateShadingImageDataType == VRSImage_Palette)
-		{
-			return VK_IMAGE_LAYOUT_FRAGMENT_SHADING_RATE_ATTACHMENT_OPTIMAL_KHR;
-		}
-		if (GRHIVariableRateShadingImageDataType == VRSImage_Fractional)
-		{
-			return VK_IMAGE_LAYOUT_FRAGMENT_DENSITY_MAP_OPTIMAL_EXT;
-		}
-	}
-
-	return VK_IMAGE_LAYOUT_UNDEFINED;
-}
-
-FVulkanRenderTargetLayout::FVulkanRenderTargetLayout(FVulkanDevice& InDevice, const FRHISetRenderTargetsInfo& RTInfo)
-	: NumAttachmentDescriptions(0)
-	, NumColorAttachments(0)
-	, bHasDepthStencil(false)
-	, bHasResolveAttachments(false)
-	, bHasDepthStencilResolve(false)
-	, bHasFragmentDensityAttachment(false)
-	, NumSamples(0)
-	, NumUsedClearValues(0)
-	, MultiViewCount(0)
-{
-	ResetAttachments();
-
-	FRenderPassCompatibleHashableStruct CompatibleHashInfo;
-	FRenderPassFullHashableStruct FullHashInfo;
-
-	bool bSetExtent = false;
-	bool bFoundClearOp = false;
-	for (int32 Index = 0; Index < RTInfo.NumColorRenderTargets; ++Index)
-	{
-		const FRHIRenderTargetView& RTView = RTInfo.ColorRenderTarget[Index];
-		if (RTView.Texture)
-		{
-			FVulkanTexture* Texture = ResourceCast(RTView.Texture);
-			check(Texture);
-			const FRHITextureDesc& TextureDesc = Texture->GetDesc();
-
-			if (bSetExtent)
-			{
-				ensure(Extent.Extent3D.width == FMath::Max(1, TextureDesc.Extent.X >> RTView.MipIndex));
-				ensure(Extent.Extent3D.height == FMath::Max(1, TextureDesc.Extent.Y >> RTView.MipIndex));
-				ensure(Extent.Extent3D.depth == TextureDesc.Depth);
-			}
-			else
-			{
-				bSetExtent = true;
-				Extent.Extent3D.width = FMath::Max(1, TextureDesc.Extent.X >> RTView.MipIndex);
-				Extent.Extent3D.height = FMath::Max(1, TextureDesc.Extent.Y >> RTView.MipIndex);
-				Extent.Extent3D.depth = TextureDesc.Depth;
-			}
-
-			ensure(!NumSamples || NumSamples == Texture->GetNumSamples());
-			NumSamples = Texture->GetNumSamples();
-		
-			VkAttachmentDescription& CurrDesc = Desc[NumAttachmentDescriptions];
-			CurrDesc.samples = static_cast<VkSampleCountFlagBits>(NumSamples);
-			CurrDesc.format = UEToVkTextureFormat(RTView.Texture->GetFormat(), EnumHasAllFlags(TextureDesc.Flags, TexCreate_SRGB));
-			CurrDesc.loadOp = RenderTargetLoadActionToVulkan(RTView.LoadAction);
-			bFoundClearOp = bFoundClearOp || (CurrDesc.loadOp == VK_ATTACHMENT_LOAD_OP_CLEAR);
-			CurrDesc.storeOp = RenderTargetStoreActionToVulkan(RTView.StoreAction);
-			CurrDesc.stencilLoadOp = VK_ATTACHMENT_LOAD_OP_DONT_CARE;
-			CurrDesc.stencilStoreOp = VK_ATTACHMENT_STORE_OP_DONT_CARE;
-
-			// Removed this temporarily as we need a way to determine if the target is actually memoryless
-			/*if (EnumHasAllFlags(Texture->UEFlags, TexCreate_Memoryless))
-			{
-				ensure(CurrDesc.storeOp == VK_ATTACHMENT_STORE_OP_DONT_CARE);
-			}*/
-
-			// If the initial != final we need to change the FullHashInfo and use FinalLayout
-			CurrDesc.initialLayout = VK_IMAGE_LAYOUT_COLOR_ATTACHMENT_OPTIMAL;
-			CurrDesc.finalLayout = VK_IMAGE_LAYOUT_COLOR_ATTACHMENT_OPTIMAL;
-
-			ColorReferences[NumColorAttachments].attachment = NumAttachmentDescriptions;
-			ColorReferences[NumColorAttachments].layout = VK_IMAGE_LAYOUT_COLOR_ATTACHMENT_OPTIMAL;
-
-			const bool bHasValidResolveAttachment = RTInfo.bHasResolveAttachments && RTInfo.ColorResolveRenderTarget[Index].Texture;
-			if (CurrDesc.samples > VK_SAMPLE_COUNT_1_BIT && bHasValidResolveAttachment)
-			{
-				Desc[NumAttachmentDescriptions + 1] = Desc[NumAttachmentDescriptions];
-				Desc[NumAttachmentDescriptions + 1].samples = VK_SAMPLE_COUNT_1_BIT;
-				Desc[NumAttachmentDescriptions + 1].loadOp = VK_ATTACHMENT_LOAD_OP_DONT_CARE;
-				Desc[NumAttachmentDescriptions + 1].storeOp = VK_ATTACHMENT_STORE_OP_STORE;
-				ResolveReferences[NumColorAttachments].attachment = NumAttachmentDescriptions + 1;
-				ResolveReferences[NumColorAttachments].layout = VK_IMAGE_LAYOUT_COLOR_ATTACHMENT_OPTIMAL;
-				CompatibleHashInfo.AttachmentsToResolve |= (uint16)(1 << NumColorAttachments);
-				++NumAttachmentDescriptions;
-				bHasResolveAttachments = true;
-			}
-
-			CompatibleHashInfo.Formats[NumColorAttachments] = CurrDesc.format;
-			FullHashInfo.LoadOps[NumColorAttachments] = CurrDesc.loadOp;
-			FullHashInfo.StoreOps[NumColorAttachments] = CurrDesc.storeOp;
-			FullHashInfo.InitialLayout[NumColorAttachments] = VK_IMAGE_LAYOUT_COLOR_ATTACHMENT_OPTIMAL;
-			++CompatibleHashInfo.NumAttachments;
-
-			++NumAttachmentDescriptions;
-			++NumColorAttachments;
-		}
-	}
-
-	if (RTInfo.DepthStencilRenderTarget.Texture)
-	{
-		VkAttachmentDescription& CurrDesc = Desc[NumAttachmentDescriptions];
-		FMemory::Memzero(CurrDesc);
-		FVulkanTexture* Texture = ResourceCast(RTInfo.DepthStencilRenderTarget.Texture);
-		check(Texture);
-		const FRHITextureDesc& TextureDesc = Texture->GetDesc();
-
-		ensure(!NumSamples || NumSamples == Texture->GetNumSamples());
-		NumSamples = TextureDesc.NumSamples;
-
-		CurrDesc.samples = static_cast<VkSampleCountFlagBits>(NumSamples);
-		CurrDesc.format = UEToVkTextureFormat(RTInfo.DepthStencilRenderTarget.Texture->GetFormat(), false);
-		CurrDesc.loadOp = RenderTargetLoadActionToVulkan(RTInfo.DepthStencilRenderTarget.DepthLoadAction);
-		CurrDesc.stencilLoadOp = RenderTargetLoadActionToVulkan(RTInfo.DepthStencilRenderTarget.StencilLoadAction);
-		bFoundClearOp = bFoundClearOp || (CurrDesc.loadOp == VK_ATTACHMENT_LOAD_OP_CLEAR || CurrDesc.stencilLoadOp == VK_ATTACHMENT_LOAD_OP_CLEAR);
-		CurrDesc.storeOp = RenderTargetStoreActionToVulkan(RTInfo.DepthStencilRenderTarget.DepthStoreAction);
-		CurrDesc.stencilStoreOp = RenderTargetStoreActionToVulkan(RTInfo.DepthStencilRenderTarget.GetStencilStoreAction());
-
-		// Removed this temporarily as we need a way to determine if the target is actually memoryless
-		/*if (EnumHasAllFlags(Texture->UEFlags, TexCreate_Memoryless))
-		{
-			ensure(CurrDesc.storeOp == VK_ATTACHMENT_STORE_OP_DONT_CARE);
-			ensure(CurrDesc.stencilStoreOp == VK_ATTACHMENT_STORE_OP_DONT_CARE);
-		}*/
-
-		const VkImageLayout DepthLayout = RTInfo.DepthStencilRenderTarget.GetDepthStencilAccess().IsDepthWrite() ? VK_IMAGE_LAYOUT_ATTACHMENT_OPTIMAL : VK_IMAGE_LAYOUT_READ_ONLY_OPTIMAL;
-		const VkImageLayout StencilLayout = RTInfo.DepthStencilRenderTarget.GetDepthStencilAccess().IsStencilWrite() ? VK_IMAGE_LAYOUT_ATTACHMENT_OPTIMAL : VK_IMAGE_LAYOUT_READ_ONLY_OPTIMAL;
-
-		// If the initial != final we need to change the FullHashInfo and use FinalLayout
-		CurrDesc.initialLayout = DepthLayout;
-		CurrDesc.finalLayout = DepthLayout;
-		StencilDesc.stencilInitialLayout = StencilLayout;
-		StencilDesc.stencilFinalLayout = StencilLayout;
-
-		DepthReference.attachment = NumAttachmentDescriptions;
-		DepthReference.layout = DepthLayout;
-		StencilReference.stencilLayout = StencilLayout;
-
-		// Use depth/stencil resolve target only if we're MSAA
-		const bool bDepthStencilResolve = (RTInfo.DepthStencilRenderTarget.DepthStoreAction == ERenderTargetStoreAction::EMultisampleResolve) || (RTInfo.DepthStencilRenderTarget.GetStencilStoreAction() == ERenderTargetStoreAction::EMultisampleResolve);
-		if (GRHISupportsDepthStencilResolve && bDepthStencilResolve && CurrDesc.samples > VK_SAMPLE_COUNT_1_BIT && RTInfo.DepthStencilResolveRenderTarget.Texture)
-		{
-			Desc[NumAttachmentDescriptions + 1] = Desc[NumAttachmentDescriptions];
-			Desc[NumAttachmentDescriptions + 1].samples = VK_SAMPLE_COUNT_1_BIT;
-			Desc[NumAttachmentDescriptions + 1].loadOp = VK_ATTACHMENT_LOAD_OP_DONT_CARE;
-			Desc[NumAttachmentDescriptions + 1].storeOp = VK_ATTACHMENT_STORE_OP_STORE;
-			Desc[NumAttachmentDescriptions + 1].stencilLoadOp = VK_ATTACHMENT_LOAD_OP_DONT_CARE;
-			Desc[NumAttachmentDescriptions + 1].stencilStoreOp = VK_ATTACHMENT_STORE_OP_STORE;
-			DepthStencilResolveReference.attachment = NumAttachmentDescriptions + 1;
-			DepthStencilResolveReference.layout = DepthLayout;
-			// NumColorAttachments was incremented after the last color attachment
-			ensureMsgf(NumColorAttachments < 16, TEXT("Must have room for depth resolve bit"));
-			CompatibleHashInfo.AttachmentsToResolve |= (uint16)(1 << NumColorAttachments);
-			++NumAttachmentDescriptions;
-			bHasDepthStencilResolve = true;
-		}
-
-		FullHashInfo.LoadOps[MaxSimultaneousRenderTargets] = CurrDesc.loadOp;
-		FullHashInfo.LoadOps[MaxSimultaneousRenderTargets + 1] = CurrDesc.stencilLoadOp;
-		FullHashInfo.StoreOps[MaxSimultaneousRenderTargets] = CurrDesc.storeOp;
-		FullHashInfo.StoreOps[MaxSimultaneousRenderTargets + 1] = CurrDesc.stencilStoreOp;
-		FullHashInfo.InitialLayout[MaxSimultaneousRenderTargets] = DepthLayout;
-		FullHashInfo.InitialLayout[MaxSimultaneousRenderTargets + 1] = StencilLayout;
-		CompatibleHashInfo.Formats[MaxSimultaneousRenderTargets] = CurrDesc.format;
-
-		++NumAttachmentDescriptions;
-
-		bHasDepthStencil = true;
-
-		if (bSetExtent)
-		{
-			// Depth can be greater or equal to color. Clamp to the smaller size.
-			Extent.Extent3D.width = FMath::Min<uint32>(Extent.Extent3D.width, TextureDesc.Extent.X);
-			Extent.Extent3D.height = FMath::Min<uint32>(Extent.Extent3D.height, TextureDesc.Extent.Y);
-		}
-		else
-		{
-			bSetExtent = true;
-			Extent.Extent3D.width = TextureDesc.Extent.X;
-			Extent.Extent3D.height = TextureDesc.Extent.Y;
-			Extent.Extent3D.depth = Texture->GetNumberOfArrayLevels();
-		}
-	}
-
-	if (GRHISupportsAttachmentVariableRateShading && RTInfo.ShadingRateTexture)
-	{
-		FVulkanTexture* Texture = ResourceCast(RTInfo.ShadingRateTexture);
-		check(Texture->GetFormat() == GRHIVariableRateShadingImageFormat);
-
-		VkAttachmentDescription& CurrDesc = Desc[NumAttachmentDescriptions];
-		FMemory::Memzero(CurrDesc);
-
-		const VkImageLayout VRSLayout = GetVRSImageLayout();
-		
-		CurrDesc.flags = 0;
-		CurrDesc.format = UEToVkTextureFormat(RTInfo.ShadingRateTexture->GetFormat(), false);
-		CurrDesc.samples = static_cast<VkSampleCountFlagBits>(RTInfo.ShadingRateTexture->GetNumSamples());
-		CurrDesc.loadOp = VK_ATTACHMENT_LOAD_OP_LOAD;
-		CurrDesc.storeOp = VK_ATTACHMENT_STORE_OP_DONT_CARE;
-		CurrDesc.stencilLoadOp = VK_ATTACHMENT_LOAD_OP_DONT_CARE;
-		CurrDesc.stencilStoreOp = VK_ATTACHMENT_STORE_OP_DONT_CARE;
-		CurrDesc.initialLayout = VRSLayout;
-		CurrDesc.finalLayout = VRSLayout;
-
-		FragmentDensityReference.attachment = NumAttachmentDescriptions;
-		FragmentDensityReference.layout = VRSLayout;
-
-		FullHashInfo.LoadOps[MaxSimultaneousRenderTargets + 2] = CurrDesc.stencilLoadOp;
-		FullHashInfo.StoreOps[MaxSimultaneousRenderTargets + 2] = CurrDesc.stencilStoreOp;
-		FullHashInfo.InitialLayout[MaxSimultaneousRenderTargets + 2] = VRSLayout;
-		CompatibleHashInfo.Formats[MaxSimultaneousRenderTargets + 1] = CurrDesc.format;
-
-		++NumAttachmentDescriptions;
-		bHasFragmentDensityAttachment = true;
-	}
-
-	SubpassHint = ESubpassHint::None;
-	CompatibleHashInfo.SubpassHint = 0;
-
-	CompatibleHashInfo.NumSamples = NumSamples;
-	CompatibleHashInfo.MultiViewCount = MultiViewCount;
-
-	RenderPassCompatibleHash = FCrc::MemCrc32(&CompatibleHashInfo, sizeof(CompatibleHashInfo));
-	RenderPassFullHash = FCrc::MemCrc32(&FullHashInfo, sizeof(FullHashInfo), RenderPassCompatibleHash);
-	NumUsedClearValues = bFoundClearOp ? NumAttachmentDescriptions : 0;
-	bCalculatedHash = true;
-}
-
-
-FVulkanRenderTargetLayout::FVulkanRenderTargetLayout(FVulkanDevice& InDevice, const FRHIRenderPassInfo& RPInfo, VkImageLayout CurrentDepthLayout, VkImageLayout CurrentStencilLayout)
-	: NumAttachmentDescriptions(0)
-	, NumColorAttachments(0)
-	, bHasDepthStencil(false)
-	, bHasResolveAttachments(false)
-	, bHasDepthStencilResolve(false)
-	, bHasFragmentDensityAttachment(false)
-	, NumSamples(0)
-	, NumUsedClearValues(0)
-	, MultiViewCount(RPInfo.MultiViewCount)
-{
-	ResetAttachments();
-
-	FRenderPassCompatibleHashableStruct CompatibleHashInfo;
-	FRenderPassFullHashableStruct FullHashInfo;
-
-	bool bSetExtent = false;
-	bool bFoundClearOp = false;
-	bool bMultiviewRenderTargets = false;
-
-	int32 NumColorRenderTargets = RPInfo.GetNumColorRenderTargets();
-	for (int32 Index = 0; Index < NumColorRenderTargets; ++Index)
-	{
-		const FRHIRenderPassInfo::FColorEntry& ColorEntry = RPInfo.ColorRenderTargets[Index];
-		FVulkanTexture* Texture = ResourceCast(ColorEntry.RenderTarget);
-		check(Texture);
-		const FRHITextureDesc& TextureDesc = Texture->GetDesc();
-
-		if (bSetExtent)
-		{
-			ensure(Extent.Extent3D.width == FMath::Max(1, TextureDesc.Extent.X >> ColorEntry.MipIndex));
-			ensure(Extent.Extent3D.height == FMath::Max(1, TextureDesc.Extent.Y >> ColorEntry.MipIndex));
-			ensure(Extent.Extent3D.depth == TextureDesc.Depth);
-		}
-		else
-		{
-			bSetExtent = true;
-			Extent.Extent3D.width = FMath::Max(1, TextureDesc.Extent.X >> ColorEntry.MipIndex);
-			Extent.Extent3D.height = FMath::Max(1, TextureDesc.Extent.Y >> ColorEntry.MipIndex);
-			Extent.Extent3D.depth = TextureDesc.Depth;
-		}
-
-		// CustomResolveSubpass can have targets with a different NumSamples
-		ensure(!NumSamples || NumSamples == ColorEntry.RenderTarget->GetNumSamples() || RPInfo.SubpassHint == ESubpassHint::CustomResolveSubpass);
-		NumSamples = ColorEntry.RenderTarget->GetNumSamples();
-
-		ensure(!GetIsMultiView() || !bMultiviewRenderTargets || Texture->GetNumberOfArrayLevels() > 1);
-		bMultiviewRenderTargets = Texture->GetNumberOfArrayLevels() > 1;
-		// With a CustomResolveSubpass last color attachment is a resolve target
-		bool bCustomResolveAttachment = (Index == (NumColorRenderTargets - 1)) && RPInfo.SubpassHint == ESubpassHint::CustomResolveSubpass;
-
-		VkAttachmentDescription& CurrDesc = Desc[NumAttachmentDescriptions];
-		CurrDesc.samples = bCustomResolveAttachment ? VK_SAMPLE_COUNT_1_BIT : static_cast<VkSampleCountFlagBits>(NumSamples);
-		CurrDesc.format = UEToVkTextureFormat(ColorEntry.RenderTarget->GetFormat(), EnumHasAllFlags(Texture->GetDesc().Flags, TexCreate_SRGB));
-		CurrDesc.loadOp = RenderTargetLoadActionToVulkan(GetLoadAction(ColorEntry.Action));
-		bFoundClearOp = bFoundClearOp || (CurrDesc.loadOp == VK_ATTACHMENT_LOAD_OP_CLEAR);
-		CurrDesc.storeOp = RenderTargetStoreActionToVulkan(GetStoreAction(ColorEntry.Action));
-		CurrDesc.stencilLoadOp = VK_ATTACHMENT_LOAD_OP_DONT_CARE;
-		CurrDesc.stencilStoreOp = VK_ATTACHMENT_STORE_OP_DONT_CARE;
-
-		if (EnumHasAnyFlags(Texture->GetDesc().Flags, TexCreate_Memoryless))
-		{
-			ensure(CurrDesc.storeOp == VK_ATTACHMENT_STORE_OP_DONT_CARE);
-		}
-
-		// If the initial != final we need to change the FullHashInfo and use FinalLayout
-		CurrDesc.initialLayout = VK_IMAGE_LAYOUT_COLOR_ATTACHMENT_OPTIMAL;
-		CurrDesc.finalLayout = VK_IMAGE_LAYOUT_COLOR_ATTACHMENT_OPTIMAL;
-
-		ColorReferences[NumColorAttachments].attachment = NumAttachmentDescriptions;
-		ColorReferences[NumColorAttachments].layout = VK_IMAGE_LAYOUT_COLOR_ATTACHMENT_OPTIMAL;
-
-		if (CurrDesc.samples > VK_SAMPLE_COUNT_1_BIT && ColorEntry.ResolveTarget)
-		{
-			Desc[NumAttachmentDescriptions + 1] = Desc[NumAttachmentDescriptions];
-			Desc[NumAttachmentDescriptions + 1].samples = VK_SAMPLE_COUNT_1_BIT;
-			Desc[NumAttachmentDescriptions + 1].loadOp = VK_ATTACHMENT_LOAD_OP_DONT_CARE;
-			Desc[NumAttachmentDescriptions + 1].storeOp = VK_ATTACHMENT_STORE_OP_STORE;
-			ResolveReferences[NumColorAttachments].attachment = NumAttachmentDescriptions + 1;
-			ResolveReferences[NumColorAttachments].layout = VK_IMAGE_LAYOUT_COLOR_ATTACHMENT_OPTIMAL;
-			CompatibleHashInfo.AttachmentsToResolve |= (uint16)(1 << NumColorAttachments);
-			++NumAttachmentDescriptions;
-			bHasResolveAttachments = true;
-		}
-
-		CompatibleHashInfo.Formats[NumColorAttachments] = CurrDesc.format;
-		FullHashInfo.LoadOps[NumColorAttachments] = CurrDesc.loadOp;
-		FullHashInfo.InitialLayout[NumColorAttachments] = VK_IMAGE_LAYOUT_COLOR_ATTACHMENT_OPTIMAL;
-		FullHashInfo.StoreOps[NumColorAttachments] = CurrDesc.storeOp;
-		++CompatibleHashInfo.NumAttachments;
-
-		++NumAttachmentDescriptions;
-		++NumColorAttachments;
-	}
-	bool bMultiViewDepthStencil = false;
-	if (RPInfo.DepthStencilRenderTarget.DepthStencilTarget)
-	{
-		VkAttachmentDescription& CurrDesc = Desc[NumAttachmentDescriptions];
-		FMemory::Memzero(CurrDesc);
-		FVulkanTexture* Texture = ResourceCast(RPInfo.DepthStencilRenderTarget.DepthStencilTarget);
-		check(Texture);
-		const FRHITextureDesc& TextureDesc = Texture->GetDesc();
-		bMultiViewDepthStencil = (Texture->GetNumberOfArrayLevels() > 1) && !Texture->GetDesc().IsTextureCube();
-		CurrDesc.samples = static_cast<VkSampleCountFlagBits>(RPInfo.DepthStencilRenderTarget.DepthStencilTarget->GetNumSamples());
-		// CustomResolveSubpass can have targets with a different NumSamples
-		ensure(!NumSamples || CurrDesc.samples == NumSamples || RPInfo.SubpassHint == ESubpassHint::CustomResolveSubpass);
-		NumSamples = CurrDesc.samples;
-		CurrDesc.format = UEToVkTextureFormat(RPInfo.DepthStencilRenderTarget.DepthStencilTarget->GetFormat(), false);
-		CurrDesc.loadOp = RenderTargetLoadActionToVulkan(GetLoadAction(GetDepthActions(RPInfo.DepthStencilRenderTarget.Action)));
-		CurrDesc.stencilLoadOp = RenderTargetLoadActionToVulkan(GetLoadAction(GetStencilActions(RPInfo.DepthStencilRenderTarget.Action)));
-		bFoundClearOp = bFoundClearOp || (CurrDesc.loadOp == VK_ATTACHMENT_LOAD_OP_CLEAR || CurrDesc.stencilLoadOp == VK_ATTACHMENT_LOAD_OP_CLEAR);
-
-		CurrDesc.storeOp = RenderTargetStoreActionToVulkan(GetStoreAction(GetDepthActions(RPInfo.DepthStencilRenderTarget.Action)));
-		CurrDesc.stencilStoreOp = RenderTargetStoreActionToVulkan(GetStoreAction(GetStencilActions(RPInfo.DepthStencilRenderTarget.Action)));
-
-		if (EnumHasAnyFlags(TextureDesc.Flags, TexCreate_Memoryless))
-		{
-			ensure(CurrDesc.storeOp == VK_ATTACHMENT_STORE_OP_DONT_CARE);
-			ensure(CurrDesc.stencilStoreOp == VK_ATTACHMENT_STORE_OP_DONT_CARE);
-		}
-
-		if (FVulkanPlatform::RequiresDepthStencilFullWrite() &&
-			Texture->GetFullAspectMask() == (VK_IMAGE_ASPECT_DEPTH_BIT | VK_IMAGE_ASPECT_STENCIL_BIT) &&
-			(CurrDesc.storeOp == VK_ATTACHMENT_STORE_OP_STORE || CurrDesc.stencilStoreOp == VK_ATTACHMENT_STORE_OP_STORE))
-		{
-			// Workaround for old mali drivers: writing not all of the image aspects to compressed render-target could cause gpu-hang
-			CurrDesc.storeOp = VK_ATTACHMENT_STORE_OP_STORE;
-			CurrDesc.stencilStoreOp = VK_ATTACHMENT_STORE_OP_STORE;
-		}
-
-		// If the initial != final we need to change the FullHashInfo and use FinalLayout
-		CurrDesc.initialLayout = CurrentDepthLayout;
-		CurrDesc.finalLayout = CurrentDepthLayout;
-		StencilDesc.stencilInitialLayout = CurrentStencilLayout;
-		StencilDesc.stencilFinalLayout = CurrentStencilLayout;
-
-		// We can't have the final layout be UNDEFINED, but it's possible that we get here from a transient texture
-		// where the stencil was never used yet.  We can set the layout to whatever we want, the next transition will
-		// happen from UNDEFINED anyhow.
-		if (CurrentDepthLayout == VK_IMAGE_LAYOUT_UNDEFINED)
-		{
-			// Unused image aspects with a LoadOp but undefined layout should just remain untouched
-			if (!RPInfo.DepthStencilRenderTarget.ExclusiveDepthStencil.IsUsingDepth() &&
-				InDevice.GetOptionalExtensions().HasEXTLoadStoreOpNone &&
-				(CurrDesc.loadOp == VK_ATTACHMENT_LOAD_OP_LOAD))
-			{
-				CurrDesc.loadOp = VK_ATTACHMENT_LOAD_OP_NONE_KHR;
-			}
-
-			check(CurrDesc.storeOp == VK_ATTACHMENT_STORE_OP_DONT_CARE);
-			CurrDesc.finalLayout = VK_IMAGE_LAYOUT_ATTACHMENT_OPTIMAL;
-		}
-		if (CurrentStencilLayout == VK_IMAGE_LAYOUT_UNDEFINED)
-		{
-			// Unused image aspects with a LoadOp but undefined layout should just remain untouched
-			if (!RPInfo.DepthStencilRenderTarget.ExclusiveDepthStencil.IsUsingStencil() &&
-				InDevice.GetOptionalExtensions().HasEXTLoadStoreOpNone &&
-				(CurrDesc.stencilLoadOp == VK_ATTACHMENT_LOAD_OP_LOAD))
-			{
-				CurrDesc.stencilLoadOp = VK_ATTACHMENT_LOAD_OP_NONE_KHR;
-			}
-
-			check(CurrDesc.stencilStoreOp == VK_ATTACHMENT_STORE_OP_DONT_CARE);
-			StencilDesc.stencilFinalLayout = VK_IMAGE_LAYOUT_ATTACHMENT_OPTIMAL;
-		}
-
-		DepthReference.attachment = NumAttachmentDescriptions;
-		DepthReference.layout = CurrentDepthLayout;
-		StencilReference.stencilLayout = CurrentStencilLayout;
-
-		if (GRHISupportsDepthStencilResolve && CurrDesc.samples > VK_SAMPLE_COUNT_1_BIT && RPInfo.DepthStencilRenderTarget.ResolveTarget)
-		{
-			Desc[NumAttachmentDescriptions + 1] = Desc[NumAttachmentDescriptions];
-			Desc[NumAttachmentDescriptions + 1].samples = VK_SAMPLE_COUNT_1_BIT;
-			Desc[NumAttachmentDescriptions + 1].loadOp = VK_ATTACHMENT_LOAD_OP_DONT_CARE;
-			Desc[NumAttachmentDescriptions + 1].storeOp = VK_ATTACHMENT_STORE_OP_STORE;
-			Desc[NumAttachmentDescriptions + 1].stencilLoadOp = VK_ATTACHMENT_LOAD_OP_DONT_CARE;
-			Desc[NumAttachmentDescriptions + 1].stencilStoreOp = VK_ATTACHMENT_STORE_OP_STORE;
-			DepthStencilResolveReference.attachment = NumAttachmentDescriptions + 1;
-			DepthStencilResolveReference.layout = CurrentDepthLayout;
-			// NumColorAttachments was incremented after the last color attachment
-			ensureMsgf(NumColorAttachments < 16, TEXT("Must have room for depth resolve bit"));
-			CompatibleHashInfo.AttachmentsToResolve |= (uint16)(1 << NumColorAttachments);
-			++NumAttachmentDescriptions;
-			bHasDepthStencilResolve = true;
-		}
-
-		FullHashInfo.LoadOps[MaxSimultaneousRenderTargets] = CurrDesc.loadOp;
-		FullHashInfo.LoadOps[MaxSimultaneousRenderTargets + 1] = CurrDesc.stencilLoadOp;
-		FullHashInfo.StoreOps[MaxSimultaneousRenderTargets] = CurrDesc.storeOp;
-		FullHashInfo.StoreOps[MaxSimultaneousRenderTargets + 1] = CurrDesc.stencilStoreOp;
-		FullHashInfo.InitialLayout[MaxSimultaneousRenderTargets] = CurrentDepthLayout;
-		FullHashInfo.InitialLayout[MaxSimultaneousRenderTargets + 1] = CurrentStencilLayout;
-		CompatibleHashInfo.Formats[MaxSimultaneousRenderTargets] = CurrDesc.format;
-
-		++NumAttachmentDescriptions;
-
-		bHasDepthStencil = true;
-
-		if (bSetExtent)
-		{
-			// Depth can be greater or equal to color. Clamp to the smaller size.
-			Extent.Extent3D.width = FMath::Min<uint32>(Extent.Extent3D.width, TextureDesc.Extent.X);
-			Extent.Extent3D.height = FMath::Min<uint32>(Extent.Extent3D.height, TextureDesc.Extent.Y);
-		}
-		else
-		{
-			bSetExtent = true;
-			Extent.Extent3D.width = TextureDesc.Extent.X;
-			Extent.Extent3D.height = TextureDesc.Extent.Y;
-			Extent.Extent3D.depth = TextureDesc.Depth;
-		}
-	}
-	else if (NumColorRenderTargets == 0)
-	{
-		// No Depth and no color, it's a raster-only pass so make sure the renderArea will be set up properly
-		checkf(RPInfo.ResolveRect.IsValid(), TEXT("For raster-only passes without render targets, ResolveRect has to contain the render area"));
-		bSetExtent = true;
-		Offset.Offset3D.x = RPInfo.ResolveRect.X1;
-		Offset.Offset3D.y = RPInfo.ResolveRect.Y1;
-		Offset.Offset3D.z = 0;
-		Extent.Extent3D.width = RPInfo.ResolveRect.X2 - RPInfo.ResolveRect.X1;
-		Extent.Extent3D.height = RPInfo.ResolveRect.Y2 - RPInfo.ResolveRect.Y1;
-		Extent.Extent3D.depth = 1;
-	}
-
-	if (GRHISupportsAttachmentVariableRateShading && RPInfo.ShadingRateTexture)
-	{
-		FVulkanTexture* Texture = ResourceCast(RPInfo.ShadingRateTexture);
-		check(Texture->GetFormat() == GRHIVariableRateShadingImageFormat);
-
-		VkAttachmentDescription& CurrDesc = Desc[NumAttachmentDescriptions];
-		FMemory::Memzero(CurrDesc);
-
-		const VkImageLayout VRSLayout = GetVRSImageLayout();
-
-		CurrDesc.flags = 0;
-		CurrDesc.format = UEToVkTextureFormat(RPInfo.ShadingRateTexture->GetFormat(), false);
-		CurrDesc.samples = static_cast<VkSampleCountFlagBits>(RPInfo.ShadingRateTexture->GetNumSamples());
-		CurrDesc.loadOp = VK_ATTACHMENT_LOAD_OP_LOAD;
-		CurrDesc.storeOp = VK_ATTACHMENT_STORE_OP_DONT_CARE;
-		CurrDesc.stencilLoadOp = VK_ATTACHMENT_LOAD_OP_DONT_CARE;
-		CurrDesc.stencilStoreOp = VK_ATTACHMENT_STORE_OP_DONT_CARE;
-		CurrDesc.initialLayout = VRSLayout;
-		CurrDesc.finalLayout = VRSLayout;
-
-		FragmentDensityReference.attachment = NumAttachmentDescriptions;
-		FragmentDensityReference.layout = VRSLayout;
-
-		FullHashInfo.LoadOps[MaxSimultaneousRenderTargets + 2] = CurrDesc.stencilLoadOp;
-		FullHashInfo.StoreOps[MaxSimultaneousRenderTargets + 2] = CurrDesc.stencilStoreOp;
-		FullHashInfo.InitialLayout[MaxSimultaneousRenderTargets + 2] = VRSLayout;
-		CompatibleHashInfo.Formats[MaxSimultaneousRenderTargets + 1] = CurrDesc.format;
-
-		++NumAttachmentDescriptions;
-		bHasFragmentDensityAttachment = true;
-	}
-
-	SubpassHint = RPInfo.SubpassHint;
-	CompatibleHashInfo.SubpassHint = (uint8)RPInfo.SubpassHint;
-
-	CompatibleHashInfo.NumSamples = NumSamples;
-	CompatibleHashInfo.MultiViewCount = MultiViewCount;
-	// Depth prepass has no color RTs but has a depth attachment that must be multiview
-	if (MultiViewCount > 1 && !bMultiviewRenderTargets && !(NumColorRenderTargets == 0 && bMultiViewDepthStencil))
-	{
-		UE_LOG(LogVulkan, Error, TEXT("Non multiview textures on a multiview layout!"));
-	}
-
-	RenderPassCompatibleHash = FCrc::MemCrc32(&CompatibleHashInfo, sizeof(CompatibleHashInfo));
-	RenderPassFullHash = FCrc::MemCrc32(&FullHashInfo, sizeof(FullHashInfo), RenderPassCompatibleHash);
-	NumUsedClearValues = bFoundClearOp ? NumAttachmentDescriptions : 0;
-	bCalculatedHash = true;
-}
-
-FVulkanRenderTargetLayout::FVulkanRenderTargetLayout(const FGraphicsPipelineStateInitializer& Initializer)
-	: NumAttachmentDescriptions(0)
-	, NumColorAttachments(0)
-	, bHasDepthStencil(false)
-	, bHasResolveAttachments(false)
-	, bHasDepthStencilResolve(false)
-	, bHasFragmentDensityAttachment(false)
-	, NumSamples(0)
-	, NumUsedClearValues(0)
-	, MultiViewCount(0)
-{
-	ResetAttachments();
-
-	FRenderPassCompatibleHashableStruct CompatibleHashInfo;
-	FRenderPassFullHashableStruct FullHashInfo;
-
-	bool bFoundClearOp = false;
-	MultiViewCount = Initializer.MultiViewCount;
-	NumSamples = Initializer.NumSamples;
-	for (uint32 Index = 0; Index < Initializer.RenderTargetsEnabled; ++Index)
-	{
-		EPixelFormat UEFormat = (EPixelFormat)Initializer.RenderTargetFormats[Index];
-		if (UEFormat != PF_Unknown)
-		{
-			// With a CustomResolveSubpass last color attachment is a resolve target
-			bool bCustomResolveAttachment = (Index == (Initializer.RenderTargetsEnabled - 1)) && Initializer.SubpassHint == ESubpassHint::CustomResolveSubpass;
-			
-			VkAttachmentDescription& CurrDesc = Desc[NumAttachmentDescriptions];
-			CurrDesc.samples = bCustomResolveAttachment ? VK_SAMPLE_COUNT_1_BIT : static_cast<VkSampleCountFlagBits>(NumSamples);
-			CurrDesc.format = UEToVkTextureFormat(UEFormat, EnumHasAllFlags(Initializer.RenderTargetFlags[Index], TexCreate_SRGB));
-			CurrDesc.loadOp = VK_ATTACHMENT_LOAD_OP_DONT_CARE;
-			CurrDesc.storeOp = VK_ATTACHMENT_STORE_OP_DONT_CARE;
-			CurrDesc.stencilLoadOp = VK_ATTACHMENT_LOAD_OP_DONT_CARE;
-			CurrDesc.stencilStoreOp = VK_ATTACHMENT_STORE_OP_DONT_CARE;
-
-			// If the initial != final we need to change the FullHashInfo and use FinalLayout
-			CurrDesc.initialLayout = VK_IMAGE_LAYOUT_COLOR_ATTACHMENT_OPTIMAL;
-			CurrDesc.finalLayout = VK_IMAGE_LAYOUT_COLOR_ATTACHMENT_OPTIMAL;
-
-			ColorReferences[NumColorAttachments].attachment = NumAttachmentDescriptions;
-			ColorReferences[NumColorAttachments].layout = VK_IMAGE_LAYOUT_COLOR_ATTACHMENT_OPTIMAL;
-
-			if (CurrDesc.samples > VK_SAMPLE_COUNT_1_BIT)
-			{
-				Desc[NumAttachmentDescriptions + 1] = Desc[NumAttachmentDescriptions];
-				Desc[NumAttachmentDescriptions + 1].samples = VK_SAMPLE_COUNT_1_BIT;
-				Desc[NumAttachmentDescriptions + 1].loadOp = VK_ATTACHMENT_LOAD_OP_DONT_CARE;
-				Desc[NumAttachmentDescriptions + 1].storeOp = VK_ATTACHMENT_STORE_OP_STORE;
-				ResolveReferences[NumColorAttachments].attachment = NumAttachmentDescriptions + 1;
-				ResolveReferences[NumColorAttachments].layout = VK_IMAGE_LAYOUT_COLOR_ATTACHMENT_OPTIMAL;
-				CompatibleHashInfo.AttachmentsToResolve |= (uint16)(1 << NumColorAttachments);
-				++NumAttachmentDescriptions;
-				bHasResolveAttachments = true;
-			}
-
-			CompatibleHashInfo.Formats[NumColorAttachments] = CurrDesc.format;
-			FullHashInfo.LoadOps[NumColorAttachments] = CurrDesc.loadOp;
-			FullHashInfo.StoreOps[NumColorAttachments] = CurrDesc.storeOp;
-			FullHashInfo.InitialLayout[NumColorAttachments] = VK_IMAGE_LAYOUT_COLOR_ATTACHMENT_OPTIMAL;
-			++CompatibleHashInfo.NumAttachments;
-
-			++NumAttachmentDescriptions;
-			++NumColorAttachments;
-		}
-	}
-
-	if (Initializer.DepthStencilTargetFormat != PF_Unknown)
-	{
-		VkAttachmentDescription& CurrDesc = Desc[NumAttachmentDescriptions];
-		FMemory::Memzero(CurrDesc);
-
-		CurrDesc.samples = static_cast<VkSampleCountFlagBits>(NumSamples);
-		CurrDesc.format = UEToVkTextureFormat(Initializer.DepthStencilTargetFormat, false);
-		CurrDesc.loadOp = RenderTargetLoadActionToVulkan(Initializer.DepthTargetLoadAction);
-		CurrDesc.stencilLoadOp = RenderTargetLoadActionToVulkan(Initializer.StencilTargetLoadAction);
-		if (CurrDesc.loadOp == VK_ATTACHMENT_LOAD_OP_CLEAR || CurrDesc.stencilLoadOp == VK_ATTACHMENT_LOAD_OP_CLEAR)
-		{
-			bFoundClearOp = true;
-		}
-		CurrDesc.storeOp = RenderTargetStoreActionToVulkan(Initializer.DepthTargetStoreAction);
-		CurrDesc.stencilStoreOp = RenderTargetStoreActionToVulkan(Initializer.StencilTargetStoreAction);
-
-		const VkImageLayout DepthLayout = Initializer.DepthStencilAccess.IsDepthWrite() ? VK_IMAGE_LAYOUT_ATTACHMENT_OPTIMAL : VK_IMAGE_LAYOUT_READ_ONLY_OPTIMAL;
-		const VkImageLayout StencilLayout = Initializer.DepthStencilAccess.IsStencilWrite() ? VK_IMAGE_LAYOUT_ATTACHMENT_OPTIMAL : VK_IMAGE_LAYOUT_READ_ONLY_OPTIMAL;
-
-		// If the initial != final we need to change the FullHashInfo and use FinalLayout
-		CurrDesc.initialLayout = DepthLayout;
-		CurrDesc.finalLayout = DepthLayout;
-		StencilDesc.stencilInitialLayout = StencilLayout;
-		StencilDesc.stencilFinalLayout = StencilLayout;
-
-		DepthReference.attachment = NumAttachmentDescriptions;
-		DepthReference.layout = DepthLayout;
-		StencilReference.stencilLayout = StencilLayout;
-
-		const bool bDepthStencilResolve = (Initializer.DepthTargetStoreAction == ERenderTargetStoreAction::EMultisampleResolve) || (Initializer.StencilTargetStoreAction == ERenderTargetStoreAction::EMultisampleResolve);
-		if (bDepthStencilResolve && GRHISupportsDepthStencilResolve && CurrDesc.samples > VK_SAMPLE_COUNT_1_BIT)
-		{
-			Desc[NumAttachmentDescriptions + 1] = Desc[NumAttachmentDescriptions];
-			Desc[NumAttachmentDescriptions + 1].samples = VK_SAMPLE_COUNT_1_BIT;
-			Desc[NumAttachmentDescriptions + 1].loadOp = VK_ATTACHMENT_LOAD_OP_DONT_CARE;
-			Desc[NumAttachmentDescriptions + 1].storeOp = VK_ATTACHMENT_STORE_OP_STORE;
-			Desc[NumAttachmentDescriptions + 1].stencilLoadOp = VK_ATTACHMENT_LOAD_OP_DONT_CARE;
-			Desc[NumAttachmentDescriptions + 1].stencilStoreOp = VK_ATTACHMENT_STORE_OP_STORE;
-			DepthStencilResolveReference.attachment = NumAttachmentDescriptions + 1;
-			DepthStencilResolveReference.layout = VK_IMAGE_LAYOUT_DEPTH_STENCIL_ATTACHMENT_OPTIMAL;
-			// NumColorAttachments was incremented after the last color attachment
-			ensureMsgf(NumColorAttachments < 16, TEXT("Must have room for depth resolve bit"));
-			CompatibleHashInfo.AttachmentsToResolve |= (uint16)(1 << NumColorAttachments);
-			++NumAttachmentDescriptions;
-			bHasDepthStencilResolve = true;
-		}
-
-		FullHashInfo.LoadOps[MaxSimultaneousRenderTargets] = CurrDesc.loadOp;
-		FullHashInfo.LoadOps[MaxSimultaneousRenderTargets + 1] = CurrDesc.stencilLoadOp;
-		FullHashInfo.StoreOps[MaxSimultaneousRenderTargets] = CurrDesc.storeOp;
-		FullHashInfo.StoreOps[MaxSimultaneousRenderTargets + 1] = CurrDesc.stencilStoreOp;
-		FullHashInfo.InitialLayout[MaxSimultaneousRenderTargets] = DepthLayout;
-		FullHashInfo.InitialLayout[MaxSimultaneousRenderTargets + 1] = StencilLayout;
-		CompatibleHashInfo.Formats[MaxSimultaneousRenderTargets] = CurrDesc.format;
-
-		++NumAttachmentDescriptions;
-		bHasDepthStencil = true;
-	}
-
-	if (Initializer.bHasFragmentDensityAttachment)
-	{
-		VkAttachmentDescription& CurrDesc = Desc[NumAttachmentDescriptions];
-		FMemory::Memzero(CurrDesc);
-
-		const VkImageLayout VRSLayout = GetVRSImageLayout();
-
-		check(GRHIVariableRateShadingImageFormat != PF_Unknown);
-
-		CurrDesc.flags = 0;
-		CurrDesc.format = UEToVkTextureFormat(GRHIVariableRateShadingImageFormat, false);
-		CurrDesc.samples = VK_SAMPLE_COUNT_1_BIT;
-		CurrDesc.loadOp = VK_ATTACHMENT_LOAD_OP_LOAD;
-		CurrDesc.storeOp = VK_ATTACHMENT_STORE_OP_DONT_CARE;
-		CurrDesc.stencilLoadOp = VK_ATTACHMENT_LOAD_OP_DONT_CARE;
-		CurrDesc.stencilStoreOp = VK_ATTACHMENT_STORE_OP_DONT_CARE;
-		CurrDesc.initialLayout = VRSLayout;
-		CurrDesc.finalLayout = VRSLayout;
-
-		FragmentDensityReference.attachment = NumAttachmentDescriptions;
-		FragmentDensityReference.layout = VRSLayout;
-
-		FullHashInfo.LoadOps[MaxSimultaneousRenderTargets + 2] = CurrDesc.stencilLoadOp;
-		FullHashInfo.StoreOps[MaxSimultaneousRenderTargets + 2] = CurrDesc.stencilStoreOp;
-		FullHashInfo.InitialLayout[MaxSimultaneousRenderTargets + 2] = VRSLayout;
-		CompatibleHashInfo.Formats[MaxSimultaneousRenderTargets + 1] = CurrDesc.format;
-
-		++NumAttachmentDescriptions;
-		bHasFragmentDensityAttachment = true;
-	}
-
-	SubpassHint = Initializer.SubpassHint;
-	CompatibleHashInfo.SubpassHint = (uint8)Initializer.SubpassHint;
-
-	CompatibleHashInfo.NumSamples = NumSamples;
-	CompatibleHashInfo.MultiViewCount = MultiViewCount;
-
-	RenderPassCompatibleHash = FCrc::MemCrc32(&CompatibleHashInfo, sizeof(CompatibleHashInfo));
-	RenderPassFullHash = FCrc::MemCrc32(&FullHashInfo, sizeof(FullHashInfo), RenderPassCompatibleHash);
-	NumUsedClearValues = bFoundClearOp ? NumAttachmentDescriptions : 0;
-	bCalculatedHash = true;
-}
diff --git a/Engine/Source/Runtime/VulkanRHI/Private/VulkanRenderpass.cpp b/Engine/Source/Runtime/VulkanRHI/Private/VulkanRenderpass.cpp
--- a/Engine/Source/Runtime/VulkanRHI/Private/VulkanRenderpass.cpp
+++ b/Engine/Source/Runtime/VulkanRHI/Private/VulkanRenderpass.cpp
@@ -1,10 +1,7 @@
 // Copyright Epic Games, Inc. All Rights Reserved.
 
-/*=============================================================================
-	VulkanState.cpp: Vulkan state implementation.
-=============================================================================*/
-
-#include "VulkanRHIPrivate.h"
+#include "VulkanRenderpass.h"
+#include "VulkanFramebuffer.h"
 #include "VulkanRenderpass.h"
 #include "VulkanContext.h"
 
diff --git a/Engine/Source/Runtime/VulkanRHI/Private/VulkanRenderpass.h b/Engine/Source/Runtime/VulkanRHI/Private/VulkanRenderpass.h
--- a/Engine/Source/Runtime/VulkanRHI/Private/VulkanRenderpass.h
+++ b/Engine/Source/Runtime/VulkanRHI/Private/VulkanRenderpass.h
@@ -9,6 +9,7 @@
 #include "CoreMinimal.h"
 #include "RHI.h"
 #include "VulkanRHIPrivate.h"
+#include "VulkanRenderTargetLayout.h"
 #include "VulkanResources.h"
 #include "VulkanPendingState.h"
 
@@ -438,6 +439,38 @@ struct FVulkanDepthStencilResolveSubpassDesc
 	}
 };
 
+class FVulkanRenderPass
+{
+public:
+	inline const FVulkanRenderTargetLayout& GetLayout() const
+	{
+		return Layout;
+	}
+
+	inline VkRenderPass GetHandle() const
+	{
+		return RenderPass;
+	}
+
+	inline uint32 GetNumUsedClearValues() const
+	{
+		return NumUsedClearValues;
+	}
+
+private:
+	friend class FVulkanRenderPassManager;
+	friend class FVulkanPipelineStateCacheManager;
+
+	FVulkanRenderPass(FVulkanDevice& Device, const FVulkanRenderTargetLayout& RTLayout);
+	~FVulkanRenderPass();
+
+private:
+	FVulkanRenderTargetLayout	Layout;
+	VkRenderPass				RenderPass;
+	uint32						NumUsedClearValues;
+	FVulkanDevice& Device;
+};
+
 extern int32 GVulkanInputAttachmentShaderRead;
 
 template <typename TSubpassDescriptionClass, typename TSubpassDependencyClass, typename TAttachmentReferenceClass, typename TAttachmentDescriptionClass, typename TRenderPassCreateInfoClass>
diff --git a/Engine/Source/Runtime/VulkanRHI/Private/VulkanResources.h b/Engine/Source/Runtime/VulkanRHI/Private/VulkanResources.h
--- a/Engine/Source/Runtime/VulkanRHI/Private/VulkanResources.h
+++ b/Engine/Source/Runtime/VulkanRHI/Private/VulkanResources.h
@@ -423,6 +423,8 @@ struct FVulkanTextureViewDesc
 	uint32 GetHash() const;
 };
 
+VkDescriptorType GetVkDescriptorTypeForViewDesc(const FRHIViewDesc& ViewDesc);
+
 class FVulkanView
 {
 public:
diff --git a/Engine/Source/Runtime/VulkanRHI/Private/VulkanShaders.cpp b/Engine/Source/Runtime/VulkanRHI/Private/VulkanShaders.cpp
--- a/Engine/Source/Runtime/VulkanRHI/Private/VulkanShaders.cpp
+++ b/Engine/Source/Runtime/VulkanRHI/Private/VulkanShaders.cpp
@@ -13,15 +13,6 @@
 #include "VulkanDescriptorSets.h"
 #include "RHICoreShader.h"
 
-TAutoConsoleVariable<int32> GDynamicGlobalUBs(
-	TEXT("r.Vulkan.DynamicGlobalUBs"),
-	2,
-	TEXT("2 to treat ALL uniform buffers as dynamic [default]\n")\
-	TEXT("1 to treat global/packed uniform buffers as dynamic\n")\
-	TEXT("0 to treat them as regular"),
-	ECVF_ReadOnly | ECVF_RenderThreadSafe
-);
-
 static int32 GVulkanCompressSPIRV = 0;
 static FAutoConsoleVariableRef GVulkanCompressSPIRVCVar(
 	TEXT("r.Vulkan.CompressSPIRV"),
@@ -621,201 +612,6 @@ FRayTracingShaderRHIRef FVulkanDynamicRHI::RHICreateRayTracingShader(TArrayView<
 	}
 }
 
-FVulkanLayout::FVulkanLayout(FVulkanDevice* InDevice, bool InGfxLayout, bool InUsesBindless)
-	: VulkanRHI::FDeviceChild(InDevice)
-	, bIsGfxLayout(InGfxLayout)
-	, bUsesBindless(InUsesBindless)
-	, DescriptorSetLayout(Device)
-	, PipelineLayout(VK_NULL_HANDLE)
-{
-}
-
-FVulkanLayout::~FVulkanLayout()
-{
-	if (PipelineLayout != VK_NULL_HANDLE)
-	{
-		Device->GetDeferredDeletionQueue().EnqueueResource(VulkanRHI::FDeferredDeletionQueue2::EType::PipelineLayout, PipelineLayout);
-		PipelineLayout = VK_NULL_HANDLE;
-	}
-}
-
-void FVulkanLayout::Compile(FVulkanDescriptorSetLayoutMap& DSetLayoutMap)
-{
-	check(PipelineLayout == VK_NULL_HANDLE);
-
-	DescriptorSetLayout.Compile(DSetLayoutMap);
-
-	if (!bUsesBindless)
-	{
-		VkPipelineLayoutCreateInfo PipelineLayoutCreateInfo;
-		ZeroVulkanStruct(PipelineLayoutCreateInfo, VK_STRUCTURE_TYPE_PIPELINE_LAYOUT_CREATE_INFO);
-
-		const TArray<VkDescriptorSetLayout>& LayoutHandles = DescriptorSetLayout.GetHandles();
-		PipelineLayoutCreateInfo.setLayoutCount = LayoutHandles.Num();
-		PipelineLayoutCreateInfo.pSetLayouts = LayoutHandles.GetData();
-		//PipelineLayoutCreateInfo.pushConstantRangeCount = 0;
-		VERIFYVULKANRESULT(VulkanRHI::vkCreatePipelineLayout(Device->GetInstanceHandle(), &PipelineLayoutCreateInfo, VULKAN_CPU_ALLOCATOR, &PipelineLayout));
-	}
-}
-
-uint32 FVulkanDescriptorSetWriter::SetupDescriptorWrites(
-	const TArray<VkDescriptorType>& Types, FVulkanHashableDescriptorInfo* InHashableDescriptorInfos,
-	VkWriteDescriptorSet* InWriteDescriptors, VkDescriptorImageInfo* InImageInfo, VkDescriptorBufferInfo* InBufferInfo, uint8* InBindingToDynamicOffsetMap,
-	VkWriteDescriptorSetAccelerationStructureKHR* InAccelerationStructuresWriteDescriptors,
-	VkAccelerationStructureKHR* InAccelerationStructures,
-	const FVulkanSamplerState& DefaultSampler, const FVulkanView::FTextureView& DefaultImageView)
-{
-	HashableDescriptorInfos = InHashableDescriptorInfos;
-	WriteDescriptors = InWriteDescriptors;
-	NumWrites = Types.Num();
-
-	BindingToDynamicOffsetMap = InBindingToDynamicOffsetMap;
-
-	InitWrittenMasks(NumWrites);
-
-	uint32 DynamicOffsetIndex = 0;
-
-	for (int32 Index = 0; Index < Types.Num(); ++Index)
-	{
-		InWriteDescriptors->sType = VK_STRUCTURE_TYPE_WRITE_DESCRIPTOR_SET;
-		InWriteDescriptors->dstBinding = Index;
-		InWriteDescriptors->descriptorCount = 1;
-		InWriteDescriptors->descriptorType = Types[Index];
-
-		switch (Types[Index])
-		{
-		case VK_DESCRIPTOR_TYPE_UNIFORM_BUFFER_DYNAMIC:
-			BindingToDynamicOffsetMap[Index] = DynamicOffsetIndex;
-			++DynamicOffsetIndex;
-			InWriteDescriptors->pBufferInfo = InBufferInfo++;
-			break;
-		case VK_DESCRIPTOR_TYPE_UNIFORM_BUFFER:
-		case VK_DESCRIPTOR_TYPE_STORAGE_BUFFER:
-			InWriteDescriptors->pBufferInfo = InBufferInfo++;
-			break;
-		case VK_DESCRIPTOR_TYPE_SAMPLER:
-			SetWrittenBase(Index); //samplers have a default setting, don't assert on those yet.
-		case VK_DESCRIPTOR_TYPE_COMBINED_IMAGE_SAMPLER:
-		case VK_DESCRIPTOR_TYPE_SAMPLED_IMAGE:
-		case VK_DESCRIPTOR_TYPE_STORAGE_IMAGE:
-		case VK_DESCRIPTOR_TYPE_INPUT_ATTACHMENT:
-			// Texture.Load() still requires a default sampler...
-			if (InHashableDescriptorInfos) // UseVulkanDescriptorCache()
-			{
-				InHashableDescriptorInfos[Index].Image.SamplerId = DefaultSampler.SamplerId;
-				InHashableDescriptorInfos[Index].Image.ImageViewId = DefaultImageView.ViewId;
-				InHashableDescriptorInfos[Index].Image.ImageLayout = static_cast<uint32>(VK_IMAGE_LAYOUT_GENERAL);
-			}
-			InImageInfo->sampler = DefaultSampler.Sampler;
-			InImageInfo->imageView = DefaultImageView.View;
-			InImageInfo->imageLayout = VK_IMAGE_LAYOUT_GENERAL;
-			InWriteDescriptors->pImageInfo = InImageInfo++;
-			break;
-		case VK_DESCRIPTOR_TYPE_STORAGE_TEXEL_BUFFER:
-		case VK_DESCRIPTOR_TYPE_UNIFORM_TEXEL_BUFFER:
-			break;
-		case VK_DESCRIPTOR_TYPE_ACCELERATION_STRUCTURE_KHR:
-			InAccelerationStructuresWriteDescriptors->sType = VK_STRUCTURE_TYPE_WRITE_DESCRIPTOR_SET_ACCELERATION_STRUCTURE_KHR;
-			InAccelerationStructuresWriteDescriptors->pNext = nullptr;
-			InAccelerationStructuresWriteDescriptors->accelerationStructureCount = 1;
-			InAccelerationStructuresWriteDescriptors->pAccelerationStructures = InAccelerationStructures++;
-			InWriteDescriptors->pNext = InAccelerationStructuresWriteDescriptors++;
-			break;
-		default:
-			checkf(0, TEXT("Unsupported descriptor type %d"), (int32)Types[Index]);
-			break;
-		}
-		++InWriteDescriptors;
-	}
-
-	return DynamicOffsetIndex;
-}
-
-void FVulkanDescriptorSetsLayoutInfo::ProcessBindingsForStage(VkShaderStageFlagBits StageFlags, ShaderStage::EStage DescSetStage, const FVulkanShaderHeader& CodeHeader, FUniformBufferGatherInfo& OutUBGatherInfo) const
-{
-	OutUBGatherInfo.CodeHeaders[DescSetStage] = &CodeHeader;
-}
-
-template<bool bIsCompute>
-void FVulkanDescriptorSetsLayoutInfo::FinalizeBindings(const FVulkanDevice& Device, const FUniformBufferGatherInfo& UBGatherInfo, const TArrayView<FRHISamplerState*>& ImmutableSamplers, bool bUsesBindless)
-{
-	// We'll be reusing this struct
-	VkDescriptorSetLayoutBinding Binding;
-	FMemory::Memzero(Binding);
-	Binding.descriptorCount = 1;
-
-	const bool bConvertAllUBsToDynamic = !bUsesBindless && (GDynamicGlobalUBs.GetValueOnAnyThread() > 1);
-	const bool bConvertPackedUBsToDynamic = !bUsesBindless && (bConvertAllUBsToDynamic || (GDynamicGlobalUBs.GetValueOnAnyThread() == 1));
-	const uint32 MaxDescriptorSetUniformBuffersDynamic = Device.GetLimits().maxDescriptorSetUniformBuffersDynamic;
-
-	int32 CurrentImmutableSampler = 0;
-	for (int32 Stage = 0; Stage < (bIsCompute ? ShaderStage::NumComputeStages : ShaderStage::NumGraphicsStages); ++Stage)
-	{
-		checkSlow(StageInfos[Stage].IsEmpty());
-
-		if (const FVulkanShaderHeader* ShaderHeader = UBGatherInfo.CodeHeaders[Stage])
-		{
-			FStageInfo& StageInfo = StageInfos[Stage];
-
-			Binding.stageFlags = UEFrequencyToVKStageBit(bIsCompute ? SF_Compute : ShaderStage::GetFrequencyForGfxStage((ShaderStage::EStage)Stage));
-
-			StageInfo.PackedGlobalsSize = ShaderHeader->PackedGlobalsSize;
-			StageInfo.NumBoundUniformBuffers = ShaderHeader->NumBoundUniformBuffers;
-
-			for (int32 BindingIndex = 0; BindingIndex < ShaderHeader->Bindings.Num(); ++BindingIndex)
-			{
-				const VkDescriptorType DescriptorType = (VkDescriptorType)ShaderHeader->Bindings[BindingIndex].DescriptorType;
-
-				const bool bIsUniformBuffer = (DescriptorType == VK_DESCRIPTOR_TYPE_UNIFORM_BUFFER);
-				const bool bIsGlobalPackedConstants = bIsUniformBuffer && ShaderHeader->PackedGlobalsSize && (BindingIndex == 0);
-
-				if (bIsGlobalPackedConstants)
-				{
-					const VkDescriptorType UBType = bConvertPackedUBsToDynamic ? VK_DESCRIPTOR_TYPE_UNIFORM_BUFFER_DYNAMIC : VK_DESCRIPTOR_TYPE_UNIFORM_BUFFER;
-
-					const uint32 NewBindingIndex = StageInfo.Types.Add(UBType);
-					checkf(NewBindingIndex == 0, TEXT("Packed globals should always be the first binding!"));
-
-					Binding.binding = NewBindingIndex;
-					Binding.descriptorType = UBType;
-					AddDescriptor(Stage, Binding);
-				}
-				else if (bIsUniformBuffer)
-				{
-					VkDescriptorType UBType = VK_DESCRIPTOR_TYPE_UNIFORM_BUFFER;
-					if (bConvertAllUBsToDynamic && LayoutTypes[VK_DESCRIPTOR_TYPE_UNIFORM_BUFFER_DYNAMIC] < MaxDescriptorSetUniformBuffersDynamic)
-					{
-						UBType = VK_DESCRIPTOR_TYPE_UNIFORM_BUFFER_DYNAMIC;
-					}
-
-					// Here we might mess up with the stageFlags, so reset them every loop
-					Binding.descriptorType = UBType;
-					const FVulkanShaderHeader::FUniformBufferInfo& UBInfo = ShaderHeader->UniformBufferInfos[BindingIndex];
-					const bool bUBHasConstantData = (BindingIndex < (int32)ShaderHeader->NumBoundUniformBuffers);
-					if (bUBHasConstantData)
-					{
-						const uint32 NewBindingIndex = StageInfo.Types.Add(UBType);
-						check(NewBindingIndex == BindingIndex);
-						Binding.binding = NewBindingIndex;
-						AddDescriptor(Stage, Binding);
-					}
-				}
-				else
-				{
-					const uint32 NewTypeIndex = StageInfo.Types.Add(DescriptorType);
-					check(NewTypeIndex == BindingIndex);
-					Binding.binding = BindingIndex;
-					Binding.descriptorType = DescriptorType;
-					AddDescriptor(Stage, Binding);
-				}
-			}
-		}
-	}
-
-	CompileTypesUsageID();
-	GenerateHash(ImmutableSamplers, bIsCompute ? VK_PIPELINE_BIND_POINT_COMPUTE : VK_PIPELINE_BIND_POINT_GRAPHICS);
-}
-
 FVulkanBoundShaderState::FVulkanBoundShaderState(FRHIVertexDeclaration* InVertexDeclarationRHI, FRHIVertexShader* InVertexShaderRHI,
 	FRHIPixelShader* InPixelShaderRHI, FRHIGeometryShader* InGeometryShaderRHI)
 	: CacheLink(InVertexDeclarationRHI, InVertexShaderRHI, InPixelShaderRHI, InGeometryShaderRHI, this)
@@ -850,7 +646,3 @@ FBoundShaderStateRHIRef FVulkanDynamicRHI::RHICreateBoundShaderState(
 
 	return new FVulkanBoundShaderState(VertexDeclarationRHI, VertexShaderRHI, PixelShaderRHI, GeometryShaderRHI);
 }
-
-
-template void FVulkanDescriptorSetsLayoutInfo::FinalizeBindings<true>(const FVulkanDevice& Device, const FUniformBufferGatherInfo& UBGatherInfo, const TArrayView<FRHISamplerState*>& ImmutableSamplers, bool bUsesBindless);
-template void FVulkanDescriptorSetsLayoutInfo::FinalizeBindings<false>(const FVulkanDevice& Device, const FUniformBufferGatherInfo& UBGatherInfo, const TArrayView<FRHISamplerState*>& ImmutableSamplers, bool bUsesBindless);
diff --git a/Engine/Source/Runtime/VulkanRHI/Private/VulkanUAV.cpp b/Engine/Source/Runtime/VulkanRHI/Private/VulkanUAV.cpp
--- a/Engine/Source/Runtime/VulkanRHI/Private/VulkanUAV.cpp
+++ b/Engine/Source/Runtime/VulkanRHI/Private/VulkanUAV.cpp
@@ -7,313 +7,16 @@
 #include "ClearReplacementShaders.h"
 #include "VulkanRayTracing.h"
 
-
-uint32 FVulkanTextureViewDesc::GetHash() const
-{
-	uint32 Hash = 0;
-	Hash = FCrc::TypeCrc32(ViewType, Hash);
-	Hash = FCrc::TypeCrc32(AspectFlags, Hash);
-	Hash = FCrc::TypeCrc32(UEFormat, Hash);
-	Hash = FCrc::TypeCrc32(Format, Hash);
-	Hash = FCrc::TypeCrc32(FirstMip, Hash);
-	Hash = FCrc::TypeCrc32(NumMips, Hash);
-	Hash = FCrc::TypeCrc32(ArraySliceIndex, Hash);
-	Hash = FCrc::TypeCrc32(NumArraySlices, Hash);
-	Hash = FCrc::TypeCrc32(bUseIdentitySwizzle, Hash);
-	Hash = FCrc::TypeCrc32(ImageUsageFlags, Hash);
-	Hash = FCrc::TypeCrc32(SamplerYcbcrConversion, Hash);
-	return Hash;
-}
-
-
-
-FVulkanView::FVulkanView(FVulkanDevice& InDevice, VkDescriptorType InDescriptorType)
-	: Device(InDevice)
-	, DescriptorType(InDescriptorType)
-{
-	BindlessHandle = Device.GetBindlessDescriptorManager()->ReserveDescriptor(InDescriptorType);
-}
-
-FVulkanView::~FVulkanView()
-{
-	Invalidate();
-
-	if (BindlessHandle.IsValid())
-	{
-		Device.GetDeferredDeletionQueue().EnqueueBindlessHandle(BindlessHandle);
-		BindlessHandle = FRHIDescriptorHandle();
-	}
-}
-
-void FVulkanView::Invalidate()
-{
-	// Carry forward its initialized state
-	const bool bIsInitialized = IsInitialized();
-
-	switch (GetViewType())
-	{
-	default: checkNoEntry(); [[fallthrough]];
-	case EType::Null:
-		break;
-
-	case EType::TypedBuffer:
-		DEC_DWORD_STAT(STAT_VulkanNumBufferViews);
-		Device.GetDeferredDeletionQueue().EnqueueResource(VulkanRHI::FDeferredDeletionQueue2::EType::BufferView, Storage.Get<FTypedBufferView>().View);
-		break;
-
-	case EType::Texture:
-		DEC_DWORD_STAT(STAT_VulkanNumImageViews);
-		Device.GetDeferredDeletionQueue().EnqueueResource(VulkanRHI::FDeferredDeletionQueue2::EType::ImageView, Storage.Get<FTextureView>().View);
-		break;
-
-	case EType::StructuredBuffer:
-		// Nothing to do
-		break;
-
-	case EType::AccelerationStructure:
-		Device.GetDeferredDeletionQueue().EnqueueResource(VulkanRHI::FDeferredDeletionQueue2::EType::AccelerationStructure, Storage.Get<FAccelerationStructureView>().Handle);
-		break;
-	}
-
-	Storage.Emplace<FInvalidatedState>();
-	Storage.Get<FInvalidatedState>().bInitialized = bIsInitialized;
-}
-
-FVulkanView* FVulkanView::InitAsTypedBufferView(FVulkanBuffer* Buffer, EPixelFormat UEFormat, uint32 InOffset, uint32 InSize)
-{
-	// We will need a deferred update if the descriptor was already in use
-	const bool bImmediateUpdate = !IsInitialized();
-
-	check(GetViewType() == EType::Null);
-	Storage.Emplace<FTypedBufferView>();
-	FTypedBufferView& TBV = Storage.Get<FTypedBufferView>();
-
-	const uint32 TotalOffset = Buffer->GetOffset() + InOffset;
-
-	check(UEFormat != PF_Unknown);
-	VkFormat Format = GVulkanBufferFormat[UEFormat];
-	check(Format != VK_FORMAT_UNDEFINED);
-
-	VkBufferViewCreateInfo ViewInfo;
-	ZeroVulkanStruct(ViewInfo, VK_STRUCTURE_TYPE_BUFFER_VIEW_CREATE_INFO);
-	ViewInfo.buffer = Buffer->GetHandle();
-	ViewInfo.offset = TotalOffset;
-	ViewInfo.format = Format;
-
-	const uint32 TypeSize =  VulkanRHI::GetNumBitsPerPixel(Format) / 8u;
-	// View size has to be a multiple of element size
-	// Commented out because there are multiple places in the high level rendering code which re-purpose buffers for a new format while there are still
-	// views with the old format lying around, and then lock them with a size computed based on the new stride, triggering this assert when the old views
-	// are re-created. These places need to be fixed before re-enabling this check (UE-211785).
-	//check(IsAligned(InSize, TypeSize));
-
-	//#todo-rco: Revisit this if buffer views become VK_BUFFER_USAGE_STORAGE_BUFFER_BIT instead of VK_BUFFER_USAGE_UNIFORM_TEXEL_BUFFER_BIT
-	const VkPhysicalDeviceLimits& Limits = Device.GetLimits();
-	const uint64 MaxSize = (uint64)Limits.maxTexelBufferElements * TypeSize;
-	ViewInfo.range = FMath::Min<uint64>(InSize, MaxSize);
-	// TODO: add a check() for exceeding MaxSize, to catch code which blindly makes views without checking the platform limits.
-
-	check(Buffer->GetBufferUsageFlags() & (VK_BUFFER_USAGE_UNIFORM_TEXEL_BUFFER_BIT | VK_BUFFER_USAGE_STORAGE_TEXEL_BUFFER_BIT));
-	check(IsAligned(InOffset, Limits.minTexelBufferOffsetAlignment));
-
-	VERIFYVULKANRESULT(VulkanRHI::vkCreateBufferView(Device.GetInstanceHandle(), &ViewInfo, VULKAN_CPU_ALLOCATOR, &TBV.View));
-
-	TBV.bVolatile = Buffer->IsVolatile();
-	if (!TBV.bVolatile && UseVulkanDescriptorCache())
-	{
-		TBV.ViewId = ++GVulkanBufferViewHandleIdCounter;
-	}
-
-	INC_DWORD_STAT(STAT_VulkanNumBufferViews);
-	// :todo-jn: the buffer view is actually not needed in bindless anymore
-
-	Device.GetBindlessDescriptorManager()->UpdateTexelBuffer(BindlessHandle, DescriptorType, ViewInfo, bImmediateUpdate);
-
-	return this;
-}
-
-FVulkanView* FVulkanView::InitAsTextureView(VkImage InImage, const FVulkanTextureViewDesc& ViewDesc)
-{
-	// We will need a deferred update if the descriptor was already in use
-	const bool bImmediateUpdate = !IsInitialized();
-
-	check(GetViewType() == EType::Null);
-	Storage.Emplace<FTextureView>();
-	FTextureView& TV = Storage.Get<FTextureView>();
-
-	LLM_SCOPE_VULKAN(ELLMTagVulkan::VulkanTextures);
-
-	VkImageViewCreateInfo ViewInfo;
-	ZeroVulkanStruct(ViewInfo, VK_STRUCTURE_TYPE_IMAGE_VIEW_CREATE_INFO);
-	ViewInfo.image = InImage;
-	ViewInfo.viewType = ViewDesc.ViewType;
-	ViewInfo.format = ViewDesc.Format;
-
-#if VULKAN_SUPPORTS_ASTC_DECODE_MODE
-	VkImageViewASTCDecodeModeEXT DecodeMode;
-	if (Device.GetOptionalExtensions().HasEXTASTCDecodeMode && IsAstcLdrFormat(ViewDesc.Format) && !IsAstcSrgbFormat(ViewDesc.Format))
-	{
-		ZeroVulkanStruct(DecodeMode, VK_STRUCTURE_TYPE_IMAGE_VIEW_ASTC_DECODE_MODE_EXT);
-		DecodeMode.decodeMode = VK_FORMAT_R8G8B8A8_UNORM;
-		DecodeMode.pNext = ViewInfo.pNext;
-		ViewInfo.pNext = &DecodeMode;
-	}
-#endif
-
-	VkSamplerYcbcrConversionInfo SamplerYcbcrConversionInfo;
-	if (ViewDesc.SamplerYcbcrConversion)
-	{
-		ZeroVulkanStruct(SamplerYcbcrConversionInfo, VK_STRUCTURE_TYPE_SAMPLER_YCBCR_CONVERSION_INFO);
-		SamplerYcbcrConversionInfo.conversion = ViewDesc.SamplerYcbcrConversion;
-		SamplerYcbcrConversionInfo.pNext = ViewInfo.pNext;
-		ViewInfo.pNext = &SamplerYcbcrConversionInfo;
-	}
-
-	if (ViewDesc.bUseIdentitySwizzle)
-	{
-		ViewInfo.components.a = VK_COMPONENT_SWIZZLE_IDENTITY;
-		ViewInfo.components.r = VK_COMPONENT_SWIZZLE_IDENTITY;
-		ViewInfo.components.g = VK_COMPONENT_SWIZZLE_IDENTITY;
-		ViewInfo.components.b = VK_COMPONENT_SWIZZLE_IDENTITY;
-	}
-	else
-	{
-		ViewInfo.components = Device.GetFormatComponentMapping(ViewDesc.UEFormat);
-	}
-
-	ViewInfo.subresourceRange.aspectMask = ViewDesc.AspectFlags;
-	ViewInfo.subresourceRange.baseMipLevel = ViewDesc.FirstMip;
-	ensure(ViewDesc.NumMips != 0xFFFFFFFF);
-	ViewInfo.subresourceRange.levelCount = ViewDesc.NumMips;
-
-	ensure(ViewDesc.ArraySliceIndex != 0xFFFFFFFF);
-	ensure(ViewDesc.NumArraySlices != 0xFFFFFFFF);
-	ViewInfo.subresourceRange.baseArrayLayer = ViewDesc.ArraySliceIndex;
-	ViewInfo.subresourceRange.layerCount = ViewDesc.NumArraySlices;
-
-	//HACK.  DX11 on PC currently uses a D24S8 depthbuffer and so needs an X24_G8 SRV to visualize stencil.
-	//So take that as our cue to visualize stencil.  In the future, the platform independent code will have a real format
-	//instead of PF_DepthStencil, so the cross-platform code could figure out the proper format to pass in for this.
-	if (ViewDesc.UEFormat == PF_X24_G8)
-	{
-		ensure((ViewInfo.format == (VkFormat)GPixelFormats[PF_DepthStencil].PlatformFormat) && (ViewInfo.format != VK_FORMAT_UNDEFINED));
-		ViewInfo.subresourceRange.aspectMask = VK_IMAGE_ASPECT_STENCIL_BIT;
-	}
-
-	// Inform the driver the view will only be used with a subset of usage flags (to help performance and/or compatibility)
-	VkImageViewUsageCreateInfo ImageViewUsageCreateInfo;
-	if (ViewDesc.ImageUsageFlags != 0)
-	{
-		ZeroVulkanStruct(ImageViewUsageCreateInfo, VK_STRUCTURE_TYPE_IMAGE_VIEW_USAGE_CREATE_INFO);
-		ImageViewUsageCreateInfo.usage = ViewDesc.ImageUsageFlags;
-
-		ImageViewUsageCreateInfo.pNext = (void*)ViewInfo.pNext;
-		ViewInfo.pNext = &ImageViewUsageCreateInfo;
-	}
-
-	INC_DWORD_STAT(STAT_VulkanNumImageViews);
-	VERIFYVULKANRESULT(VulkanRHI::vkCreateImageView(Device.GetInstanceHandle(), &ViewInfo, VULKAN_CPU_ALLOCATOR, &TV.View));
-
-	TV.Image = InImage;
-
-	if (UseVulkanDescriptorCache())
-	{
-		TV.ViewId = ++GVulkanImageViewHandleIdCounter;
-	}
-
-	const bool bDepthOrStencilAspect = (ViewDesc.AspectFlags & (VK_IMAGE_ASPECT_DEPTH_BIT | VK_IMAGE_ASPECT_STENCIL_BIT)) != 0;
-	Device.GetBindlessDescriptorManager()->UpdateImage(BindlessHandle, DescriptorType, TV.View, bDepthOrStencilAspect, bImmediateUpdate);
-
-	return this;
-}
-
-FVulkanView* FVulkanView::InitAsStructuredBufferView(FVulkanBuffer* Buffer, uint32 InOffset, uint32 InSize)
-{
-	// We will need a deferred update if the descriptor was already in use
-	const bool bImmediateUpdate = !IsInitialized();
-
-	check(GetViewType() == EType::Null);
-	Storage.Emplace<FStructuredBufferView>();
-	FStructuredBufferView& SBV = Storage.Get<FStructuredBufferView>();
-
-	const uint32 TotalOffset = Buffer->GetOffset() + InOffset;
-
-	SBV.Buffer = Buffer->GetHandle();
-	SBV.HandleId = Buffer->GetCurrentAllocation().HandleId;
-	SBV.Offset = TotalOffset;
-
-	// :todo-jn: Volatile buffers use temporary allocations that can be smaller than the buffer creation size.  Check if the savings are still worth it.
-	if (Buffer->IsVolatile())
-	{
-		InSize = FMath::Min<uint64>(InSize, Buffer->GetCurrentSize());
-	}
-
-	SBV.Size = InSize;
-
-	Device.GetBindlessDescriptorManager()->UpdateBuffer(BindlessHandle, DescriptorType, Buffer->GetHandle(), TotalOffset, InSize, bImmediateUpdate);
-
-	return this;
-}
-
-FVulkanView* FVulkanView::InitAsAccelerationStructureView(FVulkanBuffer* Buffer, uint32 Offset, uint32 Size)
-{
-	check(GetViewType() == EType::Null);
-	Storage.Emplace<FAccelerationStructureView>();
-	FAccelerationStructureView& ASV = Storage.Get<FAccelerationStructureView>();
-
-	VkAccelerationStructureCreateInfoKHR CreateInfo;
-	ZeroVulkanStruct(CreateInfo, VK_STRUCTURE_TYPE_ACCELERATION_STRUCTURE_CREATE_INFO_KHR);
-	CreateInfo.buffer = Buffer->GetHandle();
-	CreateInfo.offset = Buffer->GetOffset() + Offset;
-	CreateInfo.size = Size;
-	CreateInfo.type = VK_ACCELERATION_STRUCTURE_TYPE_TOP_LEVEL_KHR;
-
-	VERIFYVULKANRESULT(VulkanDynamicAPI::vkCreateAccelerationStructureKHR(Device.GetInstanceHandle(), &CreateInfo, VULKAN_CPU_ALLOCATOR, &ASV.Handle));
-
-	Device.GetBindlessDescriptorManager()->UpdateAccelerationStructure(BindlessHandle, ASV.Handle);
-
-	return this;
-}
-
-
-
-
-void FVulkanViewableResource::UpdateLinkedViews()
-{
-	for (FVulkanLinkedView* View = LinkedViews; View; View = View->Next())
-	{
-		View->UpdateView();
-	}
-}
-
-
-static VkImageViewType GetVkImageViewTypeForDimensionSRV(FRHIViewDesc::EDimension DescDimension, VkImageViewType TextureViewType)
-{
-	switch (DescDimension)
-	{
-	case FRHIViewDesc::EDimension::Texture2D:			return VK_IMAGE_VIEW_TYPE_2D;
-	case FRHIViewDesc::EDimension::Texture2DArray:		return VK_IMAGE_VIEW_TYPE_2D_ARRAY;
-	case FRHIViewDesc::EDimension::Texture3D:			return VK_IMAGE_VIEW_TYPE_3D;
-	case FRHIViewDesc::EDimension::TextureCube:			return VK_IMAGE_VIEW_TYPE_CUBE;
-	case FRHIViewDesc::EDimension::TextureCubeArray:	return VK_IMAGE_VIEW_TYPE_CUBE_ARRAY;
-	case FRHIViewDesc::EDimension::Unknown:				return TextureViewType;
-	default: break;
-	}
-
-	checkf(false, TEXT("Unknown texture dimension value!"));
-	return VK_IMAGE_VIEW_TYPE_MAX_ENUM;
-}
-
 static VkImageViewType GetVkImageViewTypeForDimensionUAV(FRHIViewDesc::EDimension DescDimension, VkImageViewType TextureViewType)
 {
 	switch (DescDimension)
 	{
-	case FRHIViewDesc::EDimension::Texture2D:			return VK_IMAGE_VIEW_TYPE_2D;
-	case FRHIViewDesc::EDimension::Texture2DArray:		return VK_IMAGE_VIEW_TYPE_2D_ARRAY;
-	case FRHIViewDesc::EDimension::Texture3D:			return VK_IMAGE_VIEW_TYPE_3D;
-	case FRHIViewDesc::EDimension::TextureCube:			return VK_IMAGE_VIEW_TYPE_2D_ARRAY;
-	case FRHIViewDesc::EDimension::TextureCubeArray:	return VK_IMAGE_VIEW_TYPE_2D_ARRAY;
-	case FRHIViewDesc::EDimension::Unknown:				return TextureViewType;
+	case FRHIViewDesc::EDimension::Texture2D:        return VK_IMAGE_VIEW_TYPE_2D;
+	case FRHIViewDesc::EDimension::Texture2DArray:   return VK_IMAGE_VIEW_TYPE_2D_ARRAY;
+	case FRHIViewDesc::EDimension::Texture3D:        return VK_IMAGE_VIEW_TYPE_3D;
+	case FRHIViewDesc::EDimension::TextureCube:      return VK_IMAGE_VIEW_TYPE_2D_ARRAY;
+	case FRHIViewDesc::EDimension::TextureCubeArray: return VK_IMAGE_VIEW_TYPE_2D_ARRAY;
+	case FRHIViewDesc::EDimension::Unknown:          return TextureViewType;
 	default: break;
 	}
 
@@ -321,72 +24,10 @@ static VkImageViewType GetVkImageViewTypeForDimensionUAV(FRHIViewDesc::EDimensio
 	return VK_IMAGE_VIEW_TYPE_MAX_ENUM;
 
 }
-static VkDescriptorType GetDescriptorTypeForViewDesc(FRHIViewDesc const& ViewDesc)
-{
-	if (ViewDesc.IsBuffer())
-	{
-		if (ViewDesc.IsSRV())
-		{
-			switch (ViewDesc.Buffer.SRV.BufferType)
-			{
-			case FRHIViewDesc::EBufferType::Raw:
-			case FRHIViewDesc::EBufferType::Structured:
-				return VK_DESCRIPTOR_TYPE_STORAGE_BUFFER;
-
-			case FRHIViewDesc::EBufferType::Typed:
-				return VK_DESCRIPTOR_TYPE_UNIFORM_TEXEL_BUFFER;
-
-			case FRHIViewDesc::EBufferType::AccelerationStructure:
-				return VK_DESCRIPTOR_TYPE_ACCELERATION_STRUCTURE_KHR;
-
-			default:
-				checkNoEntry();
-				return VK_DESCRIPTOR_TYPE_MAX_ENUM;
-			}
-		}
-		else
-		{
-			switch (ViewDesc.Buffer.UAV.BufferType)
-			{
-			case FRHIViewDesc::EBufferType::Raw:
-			case FRHIViewDesc::EBufferType::Structured:
-				return VK_DESCRIPTOR_TYPE_STORAGE_BUFFER;
-
-			case FRHIViewDesc::EBufferType::Typed:
-				return VK_DESCRIPTOR_TYPE_STORAGE_TEXEL_BUFFER;
-
-			case FRHIViewDesc::EBufferType::AccelerationStructure:
-				return VK_DESCRIPTOR_TYPE_ACCELERATION_STRUCTURE_KHR;
-
-			default:
-				checkNoEntry();
-				return VK_DESCRIPTOR_TYPE_MAX_ENUM;
-			}
-		}
-	}
-	else
-	{
-		if (ViewDesc.IsSRV())
-		{
-			// Sampled images aren't supported in R64, shadercompiler patches them to storage image
-			if (ViewDesc.Texture.SRV.Format == PF_R64_UINT)
-			{
-				return VK_DESCRIPTOR_TYPE_STORAGE_IMAGE;
-			}
-
-			return VK_DESCRIPTOR_TYPE_SAMPLED_IMAGE;
-		}
-		else
-		{
-			return VK_DESCRIPTOR_TYPE_STORAGE_IMAGE;
-		}
-	}
-}
 
-
-FVulkanShaderResourceView::FVulkanShaderResourceView(FRHICommandListBase& RHICmdList, FVulkanDevice& InDevice, FRHIViewableResource* InResource, FRHIViewDesc const& InViewDesc)
-	: FRHIShaderResourceView(InResource, InViewDesc)
-	, FVulkanLinkedView(InDevice, GetDescriptorTypeForViewDesc(InViewDesc))
+FVulkanUnorderedAccessView::FVulkanUnorderedAccessView(FRHICommandListBase& RHICmdList, FVulkanDevice& InDevice, FRHIViewableResource* InResource, FRHIViewDesc const& InViewDesc)
+	: FRHIUnorderedAccessView(InResource, InViewDesc)
+	, FVulkanLinkedView(InDevice, GetVkDescriptorTypeForViewDesc(InViewDesc))
 {
 	RHICmdList.EnqueueLambda([this](FRHICommandListBase&)
 	{
@@ -397,102 +38,14 @@ FVulkanShaderResourceView::FVulkanShaderResourceView(FRHICommandListBase& RHICmd
 	RHICmdList.RHIThreadFence(true);
 }
 
-FVulkanViewableResource* FVulkanShaderResourceView::GetBaseResource() const
-{
-	return IsBuffer()
-		? static_cast<FVulkanViewableResource*>(ResourceCast(GetBuffer()))
-		: static_cast<FVulkanViewableResource*>(ResourceCast(GetTexture()));
-}
-
-void FVulkanShaderResourceView::UpdateView()
+FVulkanViewableResource* FVulkanUnorderedAccessView::GetBaseResource() const
 {
-#if VULKAN_ENABLE_AGGRESSIVE_STATS
-	SCOPE_CYCLE_COUNTER(STAT_VulkanSRVUpdateTime);
-#endif
-
-	Invalidate();
-
 	if (IsBuffer())
 	{
-		FVulkanBuffer* Buffer = ResourceCast(GetBuffer());
-		auto const Info = ViewDesc.Buffer.SRV.GetViewInfo(Buffer);
-
-		if (!Info.bNullView)
-		{
-			switch (Info.BufferType)
-			{
-			case FRHIViewDesc::EBufferType::Raw:
-			case FRHIViewDesc::EBufferType::Structured:
-				InitAsStructuredBufferView(Buffer, Info.OffsetInBytes, Info.SizeInBytes);
-				break;
-
-			case FRHIViewDesc::EBufferType::Typed:
-				check(VKHasAllFlags(Buffer->GetBufferUsageFlags(), VK_BUFFER_USAGE_UNIFORM_TEXEL_BUFFER_BIT));
-				InitAsTypedBufferView(Buffer, Info.Format, Info.OffsetInBytes, Info.SizeInBytes);
-				break;
-
-			case FRHIViewDesc::EBufferType::AccelerationStructure:
-				InitAsAccelerationStructureView(Buffer, Info.OffsetInBytes, Info.SizeInBytes);
-				break;
-
-			default:
-				checkNoEntry();
-				break;
-			}
-		}
+		return ResourceCast(GetBuffer());
 	}
-	else
-	{
-		FVulkanTexture* Texture = ResourceCast(GetTexture());
-		auto const Info = ViewDesc.Texture.SRV.GetViewInfo(Texture);
-
-		uint32 ArrayFirst = Info.ArrayRange.First;
-		uint32 ArrayNum = Info.ArrayRange.Num;
-		if (Info.Dimension == FRHIViewDesc::EDimension::TextureCube || Info.Dimension == FRHIViewDesc::EDimension::TextureCubeArray)
-		{
-			ArrayFirst *= 6;
-			ArrayNum *= 6;
-			checkf((ArrayFirst + ArrayNum) <= Texture->GetNumberOfArrayLevels(), TEXT("View extends beyond original cube texture level count!"));
-		}
 
-		// Remove storage bit from usage flags for SRVs
-		const VkImageUsageFlags ImageUsageFlags = Texture->ImageUsageFlags & ~VK_IMAGE_USAGE_STORAGE_BIT;
-
-		FVulkanTextureViewDesc VulkanViewDesc;
-		VulkanViewDesc.ViewType = GetVkImageViewTypeForDimensionSRV(Info.Dimension, Texture->GetViewType());
-		VulkanViewDesc.AspectFlags = Texture->GetPartialAspectMask();
-		VulkanViewDesc.UEFormat = Info.Format;
-		VulkanViewDesc.Format = UEToVkTextureFormat(Info.Format, Info.bSRGB);
-		VulkanViewDesc.FirstMip = Info.MipRange.First;
-		VulkanViewDesc.NumMips = Info.MipRange.Num;
-		VulkanViewDesc.ArraySliceIndex = ArrayFirst;
-		VulkanViewDesc.NumArraySlices = ArrayNum;
-		VulkanViewDesc.bUseIdentitySwizzle = false;
-		VulkanViewDesc.ImageUsageFlags = ImageUsageFlags;
-		InitAsTextureView(Texture->Image, VulkanViewDesc);
-	}
-}
-
-
-
-FVulkanUnorderedAccessView::FVulkanUnorderedAccessView(FRHICommandListBase& RHICmdList, FVulkanDevice& InDevice, FRHIViewableResource* InResource, FRHIViewDesc const& InViewDesc)
-	: FRHIUnorderedAccessView(InResource, InViewDesc)
-	, FVulkanLinkedView(InDevice, GetDescriptorTypeForViewDesc(InViewDesc))
-{
-	RHICmdList.EnqueueLambda([this](FRHICommandListBase&)
-	{
-		LinkHead(GetBaseResource()->LinkedViews);
-		UpdateView();
-	});
-
-	RHICmdList.RHIThreadFence(true);
-}
-
-FVulkanViewableResource* FVulkanUnorderedAccessView::GetBaseResource() const
-{
-	return IsBuffer()
-		? static_cast<FVulkanViewableResource*>(ResourceCast(GetBuffer()))
-		: static_cast<FVulkanViewableResource*>(ResourceCast(GetTexture()));
+	return ResourceCast(GetTexture());
 }
 
 void FVulkanUnorderedAccessView::UpdateView()
@@ -564,10 +117,12 @@ void FVulkanUnorderedAccessView::UpdateView()
 
 void FVulkanUnorderedAccessView::Clear(TRHICommandList_RecursiveHazardous<FVulkanCommandListContext>& RHICmdList, const void* ClearValue, bool bFloat)
 {
-	auto GetValueType = [&](EPixelFormat Format)
+	auto GetValueType = [bFloat](EPixelFormat Format)
 	{
 		if (bFloat)
+		{
 			return EClearReplacementValueType::Float;
+		}
 
 		switch (Format)
 		{
@@ -684,11 +239,6 @@ void FVulkanCommandListContext::RHIClearUAVUint(FRHIUnorderedAccessView* Unorder
 	ResourceCast(UnorderedAccessViewRHI)->Clear(RHICmdList, &Values, false);
 }
 
-FShaderResourceViewRHIRef  FVulkanDynamicRHI::RHICreateShaderResourceView(class FRHICommandListBase& RHICmdList, FRHIViewableResource* Resource, FRHIViewDesc const& ViewDesc)
-{
-	return new FVulkanShaderResourceView(RHICmdList, *Device, Resource, ViewDesc);
-}
-
 FUnorderedAccessViewRHIRef FVulkanDynamicRHI::RHICreateUnorderedAccessView(class FRHICommandListBase& RHICmdList, FRHIViewableResource* Resource, FRHIViewDesc const& ViewDesc)
 {
 	return new FVulkanUnorderedAccessView(RHICmdList, *Device, Resource, ViewDesc);
