Reorganizing a number of Vulkan types to be in their own file(s).

diff --git a/Engine/Source/Runtime/VulkanRHI/Private/VulkanViewport.cpp b/Engine/Source/Runtime/VulkanRHI/Private/VulkanViewport.cpp
--- a/Engine/Source/Runtime/VulkanRHI/Private/VulkanViewport.cpp
+++ b/Engine/Source/Runtime/VulkanRHI/Private/VulkanViewport.cpp
@@ -290,401 +290,6 @@ void FVulkanViewport::IssueFrameEvent()
 }
 
 
-FVulkanFramebuffer::FVulkanFramebuffer(FVulkanDevice& Device, const FRHISetRenderTargetsInfo& InRTInfo, const FVulkanRenderTargetLayout& RTLayout, const FVulkanRenderPass& RenderPass)
-	: Framebuffer(VK_NULL_HANDLE)
-	, NumColorRenderTargets(InRTInfo.NumColorRenderTargets)
-	, NumColorAttachments(0)
-	, DepthStencilRenderTargetImage(VK_NULL_HANDLE)
-	, DepthStencilResolveRenderTargetImage(VK_NULL_HANDLE)
-	, FragmentDensityImage(VK_NULL_HANDLE)
-{
-	FMemory::Memzero(ColorRenderTargetImages);
-	FMemory::Memzero(ColorResolveTargetImages);
-
-	AttachmentTextureViews.Empty(RTLayout.GetNumAttachmentDescriptions());
-
-	uint32 MipIndex = 0;
-
-	const VkExtent3D& RTExtents = RTLayout.GetExtent3D();
-	// Adreno does not like zero size RTs
-	check(RTExtents.width != 0 && RTExtents.height != 0);
-	uint32 NumLayers = RTExtents.depth;
-
-	for (int32 Index = 0; Index < InRTInfo.NumColorRenderTargets; ++Index)
-	{
-		FRHITexture* RHITexture = InRTInfo.ColorRenderTarget[Index].Texture;
-		if (!RHITexture)
-		{
-			continue;
-		}
-
-		FVulkanTexture* Texture = ResourceCast(RHITexture);
-
-		ColorRenderTargetImages[Index] = Texture->Image;
-		MipIndex = InRTInfo.ColorRenderTarget[Index].MipIndex;
-
-		FVulkanView* View = GetColorRenderTargetViewDesc(Texture, MipIndex, InRTInfo.ColorRenderTarget[Index].ArraySliceIndex, RTLayout.GetMultiViewCount(), NumLayers);
-		AttachmentTextureViews.Add(View);
-
-		++NumColorAttachments;
-
-		// Check the RTLayout as well to make sure the resolve attachment is needed (Vulkan and Feature level specific)
-		// See: FVulkanRenderTargetLayout constructor with FRHIRenderPassInfo
-		if (InRTInfo.bHasResolveAttachments && RTLayout.GetHasResolveAttachments() && RTLayout.GetResolveAttachmentReferences()[Index].layout != VK_IMAGE_LAYOUT_UNDEFINED)
-		{
-			FRHITexture* ResolveRHITexture = InRTInfo.ColorResolveRenderTarget[Index].Texture;
-			FVulkanTexture* ResolveTexture = ResourceCast(ResolveRHITexture);
-			ColorResolveTargetImages[Index] = ResolveTexture->Image;
-
-			if (FVulkanView* ResolveView = GetColorResolveTargetViewDesc(ResolveTexture, MipIndex, InRTInfo.ColorRenderTarget[Index].ArraySliceIndex))
-			{
-				AttachmentTextureViews.Add(ResolveView);
-			}
-		}
-	}
-
-	if (RTLayout.GetHasDepthStencil())
-	{
-		FVulkanTexture* Texture = ResourceCast(InRTInfo.DepthStencilRenderTarget.Texture);
-		const FRHITextureDesc& Desc = Texture->GetDesc();
-		DepthStencilRenderTargetImage = Texture->Image;
-
-		check(Texture->PartialView);
-		PartialDepthTextureView = Texture->PartialView;
-
-		FVulkanView* View = GetDepthStencilTargetViewDesc(Texture, InRTInfo.NumColorRenderTargets, MipIndex, NumLayers);
-		AttachmentTextureViews.Add(View);
-
-		if (RTLayout.GetHasDepthStencilResolve() && RTLayout.GetDepthStencilResolveAttachmentReference()->layout != VK_IMAGE_LAYOUT_UNDEFINED)
-		{
-			FRHITexture* ResolveRHITexture = InRTInfo.DepthStencilResolveRenderTarget.Texture;
-			FVulkanTexture* ResolveTexture = ResourceCast(ResolveRHITexture);
-			DepthStencilResolveRenderTargetImage = ResolveTexture->Image;
-
-			if (FVulkanView* ResolveView = GetDepthStencilResolveTargetViewDesc(ResolveTexture, MipIndex))
-			{
-				AttachmentTextureViews.Add(ResolveView);
-			}
-		}
-	}
-
-	if (GRHISupportsAttachmentVariableRateShading && RTLayout.GetHasFragmentDensityAttachment())
-	{
-		FVulkanTexture* Texture = ResourceCast(InRTInfo.ShadingRateTexture);
-		FragmentDensityImage = Texture->Image;
-		FVulkanView* View = GetFragmentDensityAttachmentViewDesc(Texture, MipIndex);
-		AttachmentTextureViews.Add(View);
-	}
-
-	TArray<VkImageView> AttachmentViews;
-	AttachmentViews.Reserve(AttachmentTextureViews.Num());
-	for (FVulkanView const* View : AttachmentTextureViews)
-	{
-		AttachmentViews.Add(View->GetTextureView().View);
-	}
-
-	VkFramebufferCreateInfo CreateInfo;
-	ZeroVulkanStruct(CreateInfo, VK_STRUCTURE_TYPE_FRAMEBUFFER_CREATE_INFO);
-	CreateInfo.renderPass = RenderPass.GetHandle();
-	CreateInfo.attachmentCount = AttachmentViews.Num();
-	CreateInfo.pAttachments = AttachmentViews.GetData();
-	CreateInfo.width  = RTExtents.width;
-	CreateInfo.height = RTExtents.height;
-	CreateInfo.layers = NumLayers;
-
-	VERIFYVULKANRESULT_EXPANDED(VulkanRHI::vkCreateFramebuffer(Device.GetInstanceHandle(), &CreateInfo, VULKAN_CPU_ALLOCATOR, &Framebuffer));
-
-	RenderArea.offset.x = 0;
-	RenderArea.offset.y = 0;
-	RenderArea.extent.width = RTExtents.width;
-	RenderArea.extent.height = RTExtents.height;
-
-	INC_DWORD_STAT(STAT_VulkanNumFrameBuffers);
-}
-
-FVulkanView* FVulkanFramebuffer::GetColorRenderTargetViewDesc(FVulkanTexture* Texture, uint32 MipIndex, int32 SliceIndex, uint32 MultiViewCount, uint32& InOutNumLayers)
-{
-	check(Texture->Image != VK_NULL_HANDLE);
-
-	const FRHITextureDesc& Desc = Texture->GetDesc();
-	const VkImageUsageFlags ImageUsageFlags = VK_IMAGE_USAGE_COLOR_ATTACHMENT_BIT | (Texture->ImageUsageFlags & VK_IMAGE_USAGE_INPUT_ATTACHMENT_BIT);
-
-	if (Texture->GetViewType() == VK_IMAGE_VIEW_TYPE_2D || Texture->GetViewType() == VK_IMAGE_VIEW_TYPE_2D_ARRAY)
-	{
-		uint32 ArraySliceIndex = 0;
-		uint32 NumArraySlices = 1;
-		if (SliceIndex == -1)
-		{
-			ArraySliceIndex = 0;
-			NumArraySlices = Texture->GetNumberOfArrayLevels();
-		}
-		else
-		{
-			ArraySliceIndex = SliceIndex;
-			NumArraySlices = 1;
-			check(ArraySliceIndex < Texture->GetNumberOfArrayLevels());
-		}
-
-		// About !RTLayout.GetIsMultiView(), from https://registry.khronos.org/vulkan/specs/1.3-extensions/man/html/VkFramebufferCreateInfo.html: 
-		// If the render pass uses multiview, then layers must be one
-		if (Texture->GetViewType() == VK_IMAGE_VIEW_TYPE_2D_ARRAY && (MultiViewCount == 0))
-		{
-			InOutNumLayers = NumArraySlices;
-		}
-
-		FVulkanTextureViewDesc ViewDesc;
-		ViewDesc.ViewType = Texture->GetViewType();
-		ViewDesc.AspectFlags = Texture->GetFullAspectMask();
-		ViewDesc.UEFormat = Desc.Format;
-		ViewDesc.Format = Texture->ViewFormat;
-		ViewDesc.FirstMip = MipIndex;
-		ViewDesc.NumMips = 1;
-		ViewDesc.ArraySliceIndex = ArraySliceIndex;
-		ViewDesc.NumArraySlices = NumArraySlices;
-		ViewDesc.bUseIdentitySwizzle = true;
-		ViewDesc.ImageUsageFlags = ImageUsageFlags;
-		return Texture->FindOrAddInternalView(ViewDesc);
-	}
-	else if (Texture->GetViewType() == VK_IMAGE_VIEW_TYPE_CUBE || Texture->GetViewType() == VK_IMAGE_VIEW_TYPE_CUBE_ARRAY)
-	{
-		// Cube always renders one face at a time
-		INC_DWORD_STAT(STAT_VulkanNumImageViews);
-
-		FVulkanTextureViewDesc ViewDesc;
-		ViewDesc.ViewType = VK_IMAGE_VIEW_TYPE_2D;
-		ViewDesc.AspectFlags = Texture->GetFullAspectMask();
-		ViewDesc.UEFormat = Desc.Format;
-		ViewDesc.Format = Texture->ViewFormat;
-		ViewDesc.FirstMip = MipIndex;
-		ViewDesc.NumMips = 1;
-		ViewDesc.ArraySliceIndex = SliceIndex;
-		ViewDesc.NumArraySlices = 1;
-		ViewDesc.bUseIdentitySwizzle = true;
-		ViewDesc.ImageUsageFlags = ImageUsageFlags;
-		return Texture->FindOrAddInternalView(ViewDesc);
-	}
-	else if (Texture->GetViewType() == VK_IMAGE_VIEW_TYPE_3D)
-	{
-		FVulkanTextureViewDesc ViewDesc;
-		ViewDesc.ViewType = VK_IMAGE_VIEW_TYPE_2D_ARRAY;
-		ViewDesc.AspectFlags = Texture->GetFullAspectMask();
-		ViewDesc.UEFormat = Desc.Format;
-		ViewDesc.Format = Texture->ViewFormat;
-		ViewDesc.FirstMip = MipIndex;
-		ViewDesc.NumMips = 1;
-		ViewDesc.ArraySliceIndex = 0;
-		ViewDesc.NumArraySlices = Desc.Depth;
-		ViewDesc.bUseIdentitySwizzle = true;
-		ViewDesc.ImageUsageFlags = ImageUsageFlags;
-		return Texture->FindOrAddInternalView(ViewDesc);
-	}
-	ensure(0);
-	return nullptr;
-}
-
-FVulkanView* FVulkanFramebuffer::GetColorResolveTargetViewDesc(FVulkanTexture* ResolveTexture, uint32 MipIndex, int32 ArraySliceIndex)
-{
-	//resolve attachments only supported for 2d/2d array textures
-	if (ResolveTexture->GetViewType() == VK_IMAGE_VIEW_TYPE_2D || ResolveTexture->GetViewType() == VK_IMAGE_VIEW_TYPE_2D_ARRAY)
-	{
-		FVulkanTextureViewDesc ViewDesc;
-		ViewDesc.ViewType = ResolveTexture->GetViewType();
-		ViewDesc.AspectFlags = ResolveTexture->GetFullAspectMask();
-		ViewDesc.UEFormat = ResolveTexture->GetDesc().Format;
-		ViewDesc.Format = ResolveTexture->ViewFormat;
-		ViewDesc.FirstMip = MipIndex;
-		ViewDesc.NumMips = 1;
-		ViewDesc.ArraySliceIndex = FMath::Max(0, ArraySliceIndex);
-		ViewDesc.NumArraySlices = ResolveTexture->GetNumberOfArrayLevels();
-		ViewDesc.bUseIdentitySwizzle = true;
-		return ResolveTexture->FindOrAddInternalView(ViewDesc);
-	}
-
-	return nullptr;
-}
-
-FVulkanView* FVulkanFramebuffer::GetDepthStencilTargetViewDesc(FVulkanTexture* Texture, uint32 NumColorAttachments, uint32 MipIndex, uint32& InOutNumLayers)
-{
-	ensure(Texture->GetViewType() == VK_IMAGE_VIEW_TYPE_2D || Texture->GetViewType() == VK_IMAGE_VIEW_TYPE_2D_ARRAY || Texture->GetViewType() == VK_IMAGE_VIEW_TYPE_CUBE);
-	if ((NumColorAttachments == 0) && (Texture->GetViewType() == VK_IMAGE_VIEW_TYPE_CUBE))
-	{
-		InOutNumLayers = 6;
-
-		FVulkanTextureViewDesc ViewDesc;
-		ViewDesc.ViewType = VK_IMAGE_VIEW_TYPE_2D_ARRAY;
-		ViewDesc.AspectFlags = Texture->GetFullAspectMask();
-		ViewDesc.UEFormat = Texture->GetDesc().Format;
-		ViewDesc.Format = Texture->ViewFormat;
-		ViewDesc.FirstMip = MipIndex;
-		ViewDesc.NumMips = 1;
-		ViewDesc.ArraySliceIndex = 0;
-		ViewDesc.NumArraySlices = 6;
-		ViewDesc.bUseIdentitySwizzle = true;
-		return Texture->FindOrAddInternalView(ViewDesc);
-	}
-	else if ((Texture->GetViewType() == VK_IMAGE_VIEW_TYPE_2D) || (Texture->GetViewType() == VK_IMAGE_VIEW_TYPE_2D_ARRAY))
-	{
-		// depth attachments need a separate view to have no swizzle components, for validation correctness
-		FVulkanTextureViewDesc ViewDesc;
-		ViewDesc.ViewType = Texture->GetViewType();
-		ViewDesc.AspectFlags = Texture->GetFullAspectMask();
-		ViewDesc.UEFormat = Texture->GetDesc().Format;
-		ViewDesc.Format = Texture->ViewFormat;
-		ViewDesc.FirstMip = MipIndex;
-		ViewDesc.NumMips = 1;
-		ViewDesc.ArraySliceIndex = 0;
-		ViewDesc.NumArraySlices = Texture->GetNumberOfArrayLevels();
-		ViewDesc.bUseIdentitySwizzle = true;
-		return Texture->FindOrAddInternalView(ViewDesc);
-	}
-	else
-	{
-		return Texture->DefaultView;
-	}
-}
-
-FVulkanView* FVulkanFramebuffer::GetDepthStencilResolveTargetViewDesc(FVulkanTexture* ResolveTexture, uint32 MipIndex)
-{
-	// Resolve attachments only supported for 2d/2d array textures
-	if (ResolveTexture->GetViewType() == VK_IMAGE_VIEW_TYPE_2D || ResolveTexture->GetViewType() == VK_IMAGE_VIEW_TYPE_2D_ARRAY)
-	{
-		FVulkanTextureViewDesc ViewDesc;
-		ViewDesc.ViewType = ResolveTexture->GetViewType();
-		ViewDesc.AspectFlags = ResolveTexture->GetFullAspectMask();
-		ViewDesc.UEFormat = ResolveTexture->GetDesc().Format;
-		ViewDesc.Format = ResolveTexture->ViewFormat;
-		ViewDesc.FirstMip = MipIndex;
-		ViewDesc.NumMips = 1;
-		ViewDesc.ArraySliceIndex = 0;
-		ViewDesc.NumArraySlices = ResolveTexture->GetNumberOfArrayLevels();
-		ViewDesc.bUseIdentitySwizzle = true;
-		return ResolveTexture->FindOrAddInternalView(ViewDesc);
-	}
-
-	return nullptr;
-}
-
-FVulkanView* FVulkanFramebuffer::GetFragmentDensityAttachmentViewDesc(FVulkanTexture* Texture, uint32 MipIndex)
-{
-	ensure(Texture->GetViewType() == VK_IMAGE_VIEW_TYPE_2D || Texture->GetViewType() == VK_IMAGE_VIEW_TYPE_2D_ARRAY);
-
-	FVulkanTextureViewDesc ViewDesc;
-	ViewDesc.ViewType = Texture->GetViewType();
-	ViewDesc.AspectFlags = Texture->GetFullAspectMask();
-	ViewDesc.UEFormat = Texture->GetDesc().Format;
-	ViewDesc.Format = Texture->ViewFormat;
-	ViewDesc.FirstMip = MipIndex;
-	ViewDesc.NumMips = 1;
-	ViewDesc.ArraySliceIndex = 0;
-	ViewDesc.NumArraySlices = Texture->GetNumberOfArrayLevels();
-	ViewDesc.bUseIdentitySwizzle = true;
-	return Texture->FindOrAddInternalView(ViewDesc);
-}
-
-
-
-
-FVulkanFramebuffer::~FVulkanFramebuffer()
-{
-	ensure(Framebuffer == VK_NULL_HANDLE);
-}
-
-void FVulkanFramebuffer::Destroy(FVulkanDevice& Device)
-{
-	VulkanRHI::FDeferredDeletionQueue2& Queue = Device.GetDeferredDeletionQueue();
-	
-	// will be deleted in reverse order
-	Queue.EnqueueResource(VulkanRHI::FDeferredDeletionQueue2::EType::Framebuffer, Framebuffer);
-	Framebuffer = VK_NULL_HANDLE;
-
-	DEC_DWORD_STAT(STAT_VulkanNumFrameBuffers);
-}
-
-bool FVulkanFramebuffer::Matches(const FRHISetRenderTargetsInfo& InRTInfo) const
-{
-	if (NumColorRenderTargets != InRTInfo.NumColorRenderTargets)
-	{
-		return false;
-	}
-
-	{
-		const FRHIDepthRenderTargetView& B = InRTInfo.DepthStencilRenderTarget;
-		if (B.Texture)
-		{
-			VkImage AImage = DepthStencilRenderTargetImage;
-			VkImage BImage = ResourceCast(B.Texture)->Image;
-			if (AImage != BImage)
-			{
-				return false;
-			}
-		}
-	}
-
-	{
-		const FRHIDepthRenderTargetView& R = InRTInfo.DepthStencilResolveRenderTarget;
-		if (R.Texture)
-		{
-			VkImage AImage = DepthStencilResolveRenderTargetImage;
-			VkImage BImage = ResourceCast(R.Texture)->Image;
-			if (AImage != BImage)
-			{
-				return false;
-			}
-		}
-	}
-
-	{
-		FRHITexture* Texture = InRTInfo.ShadingRateTexture;
-		if (Texture)
-		{
-			VkImage AImage = FragmentDensityImage;
-			VkImage BImage = ResourceCast(Texture)->Image;
-			if (AImage != BImage)
-			{
-				return false;
-			}
-		}
-	}
-
-	int32 AttachementIndex = 0;
-	for (int32 Index = 0; Index < InRTInfo.NumColorRenderTargets; ++Index)
-	{
-		if (InRTInfo.bHasResolveAttachments)
-		{
-			const FRHIRenderTargetView& R = InRTInfo.ColorResolveRenderTarget[Index];
-			if (R.Texture)
-			{
-				VkImage AImage = ColorResolveTargetImages[AttachementIndex];
-				VkImage BImage = ResourceCast(R.Texture)->Image;
-				if (AImage != BImage)
-				{
-					return false;
-				}
-			}
-		}
-
-		const FRHIRenderTargetView& B = InRTInfo.ColorRenderTarget[Index];
-		if (B.Texture)
-		{
-			VkImage AImage = ColorRenderTargetImages[AttachementIndex];
-			VkImage BImage = ResourceCast(B.Texture)->Image;
-			if (AImage != BImage)
-			{
-				return false;
-			}
-			AttachementIndex++;
-		}
-	}
-
-	return true;
-}
-
-
-
-
-
-
-
 // Tear down and recreate swapchain and related resources.
 void FVulkanViewport::RecreateSwapchain(FVulkanCommandListContext& Context, FVulkanPlatformWindowContext& WindowContext)
 {
