VulkanRHI: Rename FVulkanDevice::GetInstanceHandle to FVulkanDevice::GetHandle to match with other VulkanRHI classes (and since it returns a VkDevice, not an VkInstance, which might be confusing).  Rename FVulkanRHIGraphicsPipelineState to FVulkanGraphicsPipelineState to be more consistent with the rest of VulkanRHI classes.  Remove unused variable bNeedsDynamicStateSet in CmdBuf.

#rnx

diff --git a/Engine/Source/Runtime/VulkanRHI/Private/Android/VulkanAndroidPlatform.cpp b/Engine/Source/Runtime/VulkanRHI/Private/Android/VulkanAndroidPlatform.cpp
--- a/Engine/Source/Runtime/VulkanRHI/Private/Android/VulkanAndroidPlatform.cpp
+++ b/Engine/Source/Runtime/VulkanRHI/Private/Android/VulkanAndroidPlatform.cpp
@@ -397,7 +397,7 @@ void FVulkanAndroidPlatform::InitDevice(FVulkanDevice* InDevice)
 	{
 		FVulkanQueue* GfxQueue = InDevice->GetGraphicsQueue();
 		check(GfxQueue);
-		SwappyVk_setQueueFamilyIndex(InDevice->GetInstanceHandle(), GfxQueue->GetHandle(), GfxQueue->GetFamilyIndex());
+		SwappyVk_setQueueFamilyIndex(InDevice->GetHandle(), GfxQueue->GetHandle(), GfxQueue->GetFamilyIndex());
 	}
 #endif
 }
@@ -889,13 +889,13 @@ bool FVulkanAndroidPlatform::FramePace(FVulkanDevice& Device, void* WindowHandle
 			if (CachedSyncInterval != 0)
 			{
 				// Multiple of sync interval, use swappy directly
-				SwappyVk_setSwapIntervalNS(Device.GetInstanceHandle(), Swapchain, (1000000000L) / (int64)CurrentFramePace);
+				SwappyVk_setSwapIntervalNS(Device.GetHandle(), Swapchain, (1000000000L) / (int64)CurrentFramePace);
 				bVsyncMultiple = true;
 			}
 			else
 			{
 				// Unsupported frame rate. Set to higher refresh rate and use CPU frame pacer to limit to desired frame pace
-				SwappyVk_setSwapIntervalNS(Device.GetInstanceHandle(), Swapchain, (1000000000L) / (int64)CachedRefreshRate);
+				SwappyVk_setSwapIntervalNS(Device.GetHandle(), Swapchain, (1000000000L) / (int64)CachedRefreshRate);
 				// indicate that the RHI should perform CPU frame pacing to handle the requested frame rate
 				bVsyncMultiple = false;
 			}
@@ -1897,7 +1897,7 @@ VkPipelineCache FVulkanAndroidPlatform::PrecompilePSO(
 							PipelineCacheCreateInfo.flags = 0;
 							PipelineCacheCreateInfo.pInitialData = ResultSharedBuffer + sizeof(ResultMemSize);
 							PipelineCacheCreateInfo.initialDataSize = ResultSize;
-							VERIFYVULKANRESULT(VulkanRHI::vkCreatePipelineCache(Device->GetInstanceHandle(), &PipelineCacheCreateInfo, VULKAN_CPU_ALLOCATOR, &ReturnPipelineCache));
+							VERIFYVULKANRESULT(VulkanRHI::vkCreatePipelineCache(Device->GetHandle(), &PipelineCacheCreateInfo, VULKAN_CPU_ALLOCATOR, &ReturnPipelineCache));
 							AfterSize = ResultSize - OptionalPSOCacheData.Num();
 						}
 					}
@@ -2015,16 +2015,16 @@ void FVulkanAndroidPlatform::SetupImageMemoryRequirementWorkaround(const FVulkan
 
 		VkImage Image0;
 		VkMemoryRequirements Image0Mem;
-		VERIFYVULKANRESULT(VulkanRHI::vkCreateImage(InDevice.GetInstanceHandle(), &ImageCreateInfo, VULKAN_CPU_ALLOCATOR, &Image0));
-		VulkanRHI::vkGetImageMemoryRequirements(InDevice.GetInstanceHandle(), Image0, &Image0Mem);
-		VulkanRHI::vkDestroyImage(InDevice.GetInstanceHandle(), Image0, VULKAN_CPU_ALLOCATOR);
+		VERIFYVULKANRESULT(VulkanRHI::vkCreateImage(InDevice.GetHandle(), &ImageCreateInfo, VULKAN_CPU_ALLOCATOR, &Image0));
+		VulkanRHI::vkGetImageMemoryRequirements(InDevice.GetHandle(), Image0, &Image0Mem);
+		VulkanRHI::vkDestroyImage(InDevice.GetHandle(), Image0, VULKAN_CPU_ALLOCATOR);
 
 		VkImage ImageMutable;
 		VkMemoryRequirements ImageMutableMem;
 		ImageCreateInfo.flags = VK_IMAGE_CREATE_MUTABLE_FORMAT_BIT;
-		VERIFYVULKANRESULT(VulkanRHI::vkCreateImage(InDevice.GetInstanceHandle(), &ImageCreateInfo, VULKAN_CPU_ALLOCATOR, &ImageMutable));
-		VulkanRHI::vkGetImageMemoryRequirements(InDevice.GetInstanceHandle(), ImageMutable, &ImageMutableMem);
-		VulkanRHI::vkDestroyImage(InDevice.GetInstanceHandle(), ImageMutable, VULKAN_CPU_ALLOCATOR);
+		VERIFYVULKANRESULT(VulkanRHI::vkCreateImage(InDevice.GetHandle(), &ImageCreateInfo, VULKAN_CPU_ALLOCATOR, &ImageMutable));
+		VulkanRHI::vkGetImageMemoryRequirements(InDevice.GetHandle(), ImageMutable, &ImageMutableMem);
+		VulkanRHI::vkDestroyImage(InDevice.GetHandle(), ImageMutable, VULKAN_CPU_ALLOCATOR);
 
 		VkImage ImageStorage;
 		VkMemoryRequirements ImageStorageMem;
@@ -2033,9 +2033,9 @@ void FVulkanAndroidPlatform::SetupImageMemoryRequirementWorkaround(const FVulkan
 			ImageCreateInfo.usage |= VK_IMAGE_USAGE_STORAGE_BIT;
 			ImageCreateInfo.flags = 0;
 		}
-		VERIFYVULKANRESULT(VulkanRHI::vkCreateImage(InDevice.GetInstanceHandle(), &ImageCreateInfo, VULKAN_CPU_ALLOCATOR, &ImageStorage));
-		VulkanRHI::vkGetImageMemoryRequirements(InDevice.GetInstanceHandle(), ImageStorage, &ImageStorageMem);
-		VulkanRHI::vkDestroyImage(InDevice.GetInstanceHandle(), ImageStorage, VULKAN_CPU_ALLOCATOR);
+		VERIFYVULKANRESULT(VulkanRHI::vkCreateImage(InDevice.GetHandle(), &ImageCreateInfo, VULKAN_CPU_ALLOCATOR, &ImageStorage));
+		VulkanRHI::vkGetImageMemoryRequirements(InDevice.GetHandle(), ImageStorage, &ImageStorageMem);
+		VulkanRHI::vkDestroyImage(InDevice.GetHandle(), ImageStorage, VULKAN_CPU_ALLOCATOR);
 
 		const float MEM_SIZE_THRESHOLD = 1.5f;
 		const float IMAGE0_SIZE = (float)Image0Mem.size;
@@ -2071,17 +2071,17 @@ void FVulkanAndroidPlatform::SetupImageMemoryRequirementWorkaround(const FVulkan
 
 		VkImage ImageOptimal_ASTC;
 		VkMemoryRequirements ImageOptimalMem_ASTC;
-		VERIFYVULKANRESULT(VulkanRHI::vkCreateImage(InDevice.GetInstanceHandle(), &ImageCreateInfo, VULKAN_CPU_ALLOCATOR, &ImageOptimal_ASTC));
-		VulkanRHI::vkGetImageMemoryRequirements(InDevice.GetInstanceHandle(), ImageOptimal_ASTC, &ImageOptimalMem_ASTC);
-		VulkanRHI::vkDestroyImage(InDevice.GetInstanceHandle(), ImageOptimal_ASTC, VULKAN_CPU_ALLOCATOR);
+		VERIFYVULKANRESULT(VulkanRHI::vkCreateImage(InDevice.GetHandle(), &ImageCreateInfo, VULKAN_CPU_ALLOCATOR, &ImageOptimal_ASTC));
+		VulkanRHI::vkGetImageMemoryRequirements(InDevice.GetHandle(), ImageOptimal_ASTC, &ImageOptimalMem_ASTC);
+		VulkanRHI::vkDestroyImage(InDevice.GetHandle(), ImageOptimal_ASTC, VULKAN_CPU_ALLOCATOR);
 
 		ImageCreateInfo.tiling = VK_IMAGE_TILING_LINEAR;
 
 		VkImage ImageLinear_ASTC;
 		VkMemoryRequirements ImageLinearMem_ASTC;
-		VERIFYVULKANRESULT(VulkanRHI::vkCreateImage(InDevice.GetInstanceHandle(), &ImageCreateInfo, VULKAN_CPU_ALLOCATOR, &ImageLinear_ASTC));
-		VulkanRHI::vkGetImageMemoryRequirements(InDevice.GetInstanceHandle(), ImageLinear_ASTC, &ImageLinearMem_ASTC);
-		VulkanRHI::vkDestroyImage(InDevice.GetInstanceHandle(), ImageLinear_ASTC, VULKAN_CPU_ALLOCATOR);
+		VERIFYVULKANRESULT(VulkanRHI::vkCreateImage(InDevice.GetHandle(), &ImageCreateInfo, VULKAN_CPU_ALLOCATOR, &ImageLinear_ASTC));
+		VulkanRHI::vkGetImageMemoryRequirements(InDevice.GetHandle(), ImageLinear_ASTC, &ImageLinearMem_ASTC);
+		VulkanRHI::vkDestroyImage(InDevice.GetHandle(), ImageLinear_ASTC, VULKAN_CPU_ALLOCATOR);
 
 
 		const float MEM_SIZE_THRESHOLD = 2.0f;
diff --git a/Engine/Source/Runtime/VulkanRHI/Private/VulkanBindlessDescriptorManager.cpp b/Engine/Source/Runtime/VulkanRHI/Private/VulkanBindlessDescriptorManager.cpp
--- a/Engine/Source/Runtime/VulkanRHI/Private/VulkanBindlessDescriptorManager.cpp
+++ b/Engine/Source/Runtime/VulkanRHI/Private/VulkanBindlessDescriptorManager.cpp
@@ -308,7 +308,7 @@ FVulkanBindlessDescriptorManager::~FVulkanBindlessDescriptorManager()
 
 void FVulkanBindlessDescriptorManager::Deinit()
 {
-	const VkDevice DeviceHandle = Device.GetInstanceHandle();
+	const VkDevice DeviceHandle = Device.GetHandle();
 
 	if (bIsSupported)
 	{
@@ -355,7 +355,7 @@ void FVulkanBindlessDescriptorManager::Init()
 		return;
 	}
 
-	const VkDevice DeviceHandle = Device.GetInstanceHandle();
+	const VkDevice DeviceHandle = Device.GetHandle();
 	const VkPhysicalDeviceDescriptorBufferPropertiesEXT& DescriptorBufferProperties = Device.GetOptionalExtensionProperties().DescriptorBufferProps;
 
 	const VkBufferUsageFlags BufferUsageFlags = VK_BUFFER_USAGE_SHADER_DEVICE_ADDRESS_BIT | VK_BUFFER_USAGE_RESOURCE_DESCRIPTOR_BUFFER_BIT_EXT;
@@ -660,7 +660,7 @@ void FVulkanBindlessDescriptorManager::RegisterUniformBuffers(FVulkanCommandList
 				const int32 BindingIndex = StageOffset + DescriptorAddressInfoIndex;
 				VkDeviceSize BindingByteOffset = 0;
 #if UE_BUILD_DEBUG
-				VulkanRHI::vkGetDescriptorSetLayoutBindingOffsetEXT(Device.GetInstanceHandle(), SingleUseUBDescriptorSetLayout, BindingIndex, &BindingByteOffset);
+				VulkanRHI::vkGetDescriptorSetLayoutBindingOffsetEXT(Device.GetHandle(), SingleUseUBDescriptorSetLayout, BindingIndex, &BindingByteOffset);
 				check(BindingByteOffset == BindingIndex * UBDescriptorSize);
 #else
 				BindingByteOffset = (BindingIndex * UBDescriptorSize);
@@ -670,7 +670,7 @@ void FVulkanBindlessDescriptorManager::RegisterUniformBuffers(FVulkanCommandList
 				ZeroVulkanStruct(Info, VK_STRUCTURE_TYPE_DESCRIPTOR_GET_INFO_EXT);
 				Info.type = VK_DESCRIPTOR_TYPE_UNIFORM_BUFFER;
 				Info.data.pUniformBuffer = &DescriptorAddressInfo;
-				VulkanRHI::vkGetDescriptorEXT(Device.GetInstanceHandle(), &Info, UBDescriptorSize, &MappedPointer[BindingByteOffset]);
+				VulkanRHI::vkGetDescriptorEXT(Device.GetHandle(), &Info, UBDescriptorSize, &MappedPointer[BindingByteOffset]);
 			}
 		}
 	}
@@ -729,7 +729,7 @@ void FVulkanBindlessDescriptorManager::UpdateDescriptor(FRHIDescriptorHandle Des
 	ZeroVulkanStruct(Info, VK_STRUCTURE_TYPE_DESCRIPTOR_GET_INFO_EXT);
 	Info.type = DescriptorType;
 	Info.data = DescriptorData;
-	VulkanRHI::vkGetDescriptorEXT(Device.GetInstanceHandle(), &Info, State.DescriptorSize, &State.DebugDescriptors[ByteOffset]);
+	VulkanRHI::vkGetDescriptorEXT(Device.GetHandle(), &Info, State.DescriptorSize, &State.DebugDescriptors[ByteOffset]);
 
 	if (bImmediateUpdate)
 	{
@@ -819,7 +819,7 @@ void FVulkanBindlessDescriptorManager::UpdateBuffer(FRHIDescriptorHandle Descrip
 		VkBufferDeviceAddressInfo BufferInfo;
 		ZeroVulkanStruct(BufferInfo, VK_STRUCTURE_TYPE_BUFFER_DEVICE_ADDRESS_INFO);
 		BufferInfo.buffer = VulkanBuffer;
-		const VkDeviceAddress BufferAddress = VulkanRHI::vkGetBufferDeviceAddressKHR(Device.GetInstanceHandle(), &BufferInfo);
+		const VkDeviceAddress BufferAddress = VulkanRHI::vkGetBufferDeviceAddressKHR(Device.GetHandle(), &BufferInfo);
 
 		UpdateBuffer(DescriptorHandle, DescriptorType, BufferAddress + BufferOffset, BufferSize, bImmediateUpdate);
 	}
@@ -852,7 +852,7 @@ void FVulkanBindlessDescriptorManager::UpdateTexelBuffer(FRHIDescriptorHandle De
 		VkBufferDeviceAddressInfo BufferInfo;
 		ZeroVulkanStruct(BufferInfo, VK_STRUCTURE_TYPE_BUFFER_DEVICE_ADDRESS_INFO);
 		BufferInfo.buffer = ViewInfo.buffer;
-		const VkDeviceAddress BufferAddress = VulkanRHI::vkGetBufferDeviceAddressKHR(Device.GetInstanceHandle(), &BufferInfo);
+		const VkDeviceAddress BufferAddress = VulkanRHI::vkGetBufferDeviceAddressKHR(Device.GetHandle(), &BufferInfo);
 
 		VkDescriptorAddressInfoEXT AddressInfo;
 		ZeroVulkanStruct(AddressInfo, VK_STRUCTURE_TYPE_DESCRIPTOR_ADDRESS_INFO_EXT);
@@ -874,7 +874,7 @@ void FVulkanBindlessDescriptorManager::UpdateAccelerationStructure(FRHIDescripto
 		VkAccelerationStructureDeviceAddressInfoKHR AccelerationStructureDeviceAddressInfo;
 		ZeroVulkanStruct(AccelerationStructureDeviceAddressInfo, VK_STRUCTURE_TYPE_ACCELERATION_STRUCTURE_DEVICE_ADDRESS_INFO_KHR);
 		AccelerationStructureDeviceAddressInfo.accelerationStructure = AccelerationStructure;
-		const VkDeviceAddress BufferAddress = VulkanRHI::vkGetAccelerationStructureDeviceAddressKHR(Device.GetInstanceHandle(), &AccelerationStructureDeviceAddressInfo);
+		const VkDeviceAddress BufferAddress = VulkanRHI::vkGetAccelerationStructureDeviceAddressKHR(Device.GetHandle(), &AccelerationStructureDeviceAddressInfo);
 
 		VkDescriptorDataEXT DescriptorData;
 		DescriptorData.accelerationStructure = BufferAddress;
diff --git a/Engine/Source/Runtime/VulkanRHI/Private/VulkanBuffer.cpp b/Engine/Source/Runtime/VulkanRHI/Private/VulkanBuffer.cpp
--- a/Engine/Source/Runtime/VulkanRHI/Private/VulkanBuffer.cpp
+++ b/Engine/Source/Runtime/VulkanRHI/Private/VulkanBuffer.cpp
@@ -68,7 +68,7 @@ static VkDeviceAddress GetBufferDeviceAddress(FVulkanDevice* Device, VkBuffer Bu
 		VkBufferDeviceAddressInfoKHR DeviceAddressInfo;
 		ZeroVulkanStruct(DeviceAddressInfo, VK_STRUCTURE_TYPE_BUFFER_DEVICE_ADDRESS_INFO);
 		DeviceAddressInfo.buffer = Buffer;
-		return VulkanRHI::vkGetBufferDeviceAddressKHR(Device->GetInstanceHandle(), &DeviceAddressInfo);
+		return VulkanRHI::vkGetBufferDeviceAddressKHR(Device->GetHandle(), &DeviceAddressInfo);
 	}
 	return 0;
 }
diff --git a/Engine/Source/Runtime/VulkanRHI/Private/VulkanChunkedPipelineCache.cpp b/Engine/Source/Runtime/VulkanRHI/Private/VulkanChunkedPipelineCache.cpp
--- a/Engine/Source/Runtime/VulkanRHI/Private/VulkanChunkedPipelineCache.cpp
+++ b/Engine/Source/Runtime/VulkanRHI/Private/VulkanChunkedPipelineCache.cpp
@@ -80,7 +80,7 @@ namespace UE
 
 		static VkDevice GetVulkanDeviceHandle()
 		{
-			return GetVulkanDevice()->GetInstanceHandle();
+			return GetVulkanDevice()->GetHandle();
 		}
 
 		int32 GMaxPSOsPerChunk				= 20;
@@ -184,8 +184,8 @@ namespace UE
 		}
 
 		using FPipelineCacheChunkKey = uint32;
-		using FVulkanRHIGraphicsPipelineStateLRU = TDoubleLinkedList<FPipelineCacheChunkKey>;
-		using FVulkanPipelineCacheChunkLRUNode = FVulkanRHIGraphicsPipelineStateLRU::TDoubleLinkedListNode;	
+		using FVulkanGraphicsPipelineStateLRU = TDoubleLinkedList<FPipelineCacheChunkKey>;
+		using FVulkanPipelineCacheChunkLRUNode = FVulkanGraphicsPipelineStateLRU::TDoubleLinkedListNode;
 
 		std::atomic<int32> TotalResidentCacheSize = 0;
 	}
@@ -751,7 +751,7 @@ public:
 			VERIFYVULKANRESULT(VulkanRHI::vkCreatePipelineCache(UE::Vulkan::GetVulkanDeviceHandle(), &PipelineCacheInfo, VULKAN_CPU_ALLOCATOR, &PipelineCacheObj));
 		}
 
-		uint64 GetPrecacheKey(const FVulkanRHIGraphicsPipelineState* GraphicsPipelineState)
+		uint64 GetPrecacheKey(const FVulkanGraphicsPipelineState* GraphicsPipelineState)
 		{
 			return GraphicsPipelineState->PrecacheKey;
 		}
@@ -897,9 +897,9 @@ class FVulkanChunkedPipelineCacheManagerImpl
 
 	FCriticalSection LRUCS;
 	using FVulkanPipelineCacheChunkLRUNode = UE::Vulkan::FVulkanPipelineCacheChunkLRUNode;
-	using FVulkanRHIGraphicsPipelineStateLRU = UE::Vulkan::FVulkanRHIGraphicsPipelineStateLRU;
+	using FVulkanGraphicsPipelineStateLRU = UE::Vulkan::FVulkanGraphicsPipelineStateLRU;
 
-	UE::Vulkan::FVulkanRHIGraphicsPipelineStateLRU CacheChunkLRU;
+	UE::Vulkan::FVulkanGraphicsPipelineStateLRU CacheChunkLRU;
 
 
 
@@ -1065,7 +1065,7 @@ private:
 		return bSuccess;
 	}
 
-	static uint64 GetPrecacheHash(const FVulkanRHIGraphicsPipelineState* GFXState) { return GFXState->PrecacheKey; }
+	static uint64 GetPrecacheHash(const FVulkanGraphicsPipelineState* GFXState) { return GFXState->PrecacheKey; }
 
 	template<class TPipelineState, class TInitializer>
 	FVulkanPipelineCacheChunk* GetChunk(FVulkanChunkedPipelineCacheManager::FPSOCreateCallbackFunc<TPipelineState>& PSOCreateFunc, const TInitializer& Initializer, FVulkanPipelineCacheChunk::EPSOCacheFindResult& FindResult, FCacheChunkKey& ChunkKeyOUT, FRWScopeLock& Lock, const bool bTryAdd = false)
@@ -1249,5 +1249,5 @@ void FVulkanChunkedPipelineCacheManager::Tick()
 	VulkanPipelineCacheManagerImpl->Tick();
 }
 
-template VkResult FVulkanChunkedPipelineCacheManager::CreatePSO(FVulkanRHIGraphicsPipelineState* GraphicsPipelineState, bool bIsPrecompileJob, FPSOCreateCallbackFunc<FVulkanRHIGraphicsPipelineState> PSOCreateFunc);
+template VkResult FVulkanChunkedPipelineCacheManager::CreatePSO(FVulkanGraphicsPipelineState* GraphicsPipelineState, bool bIsPrecompileJob, FPSOCreateCallbackFunc<FVulkanGraphicsPipelineState> PSOCreateFunc);
 
diff --git a/Engine/Source/Runtime/VulkanRHI/Private/VulkanCommandBuffer.cpp b/Engine/Source/Runtime/VulkanRHI/Private/VulkanCommandBuffer.cpp
--- a/Engine/Source/Runtime/VulkanRHI/Private/VulkanCommandBuffer.cpp
+++ b/Engine/Source/Runtime/VulkanRHI/Private/VulkanCommandBuffer.cpp
@@ -41,9 +41,8 @@ void FVulkanCommandBuffer::AllocMemory()
 	CreateCmdBufInfo.commandBufferCount = 1;
 	CreateCmdBufInfo.commandPool = CommandBufferPool.GetHandle();
 
-	VERIFYVULKANRESULT(VulkanRHI::vkAllocateCommandBuffers(Device.GetInstanceHandle(), &CreateCmdBufInfo, &CommandBufferHandle));
+	VERIFYVULKANRESULT(VulkanRHI::vkAllocateCommandBuffers(Device.GetHandle(), &CreateCmdBufInfo, &CommandBufferHandle));
 
-	bNeedsDynamicStateSet = 1;
 	bHasPipeline = 0;
 	bHasViewport = 0;
 	bHasScissor = 0;
@@ -66,7 +65,7 @@ void FVulkanCommandBuffer::FreeMemory()
 	// Assumes we are inside a lock for the pool
 	check(State != EState::NotAllocated);
 	check(CommandBufferHandle != VK_NULL_HANDLE);
-	VulkanRHI::vkFreeCommandBuffers(Device.GetInstanceHandle(), CommandBufferPool.GetHandle(), 1, &CommandBufferHandle);
+	VulkanRHI::vkFreeCommandBuffers(Device.GetHandle(), CommandBufferPool.GetHandle(), 1, &CommandBufferHandle);
 	CommandBufferHandle = VK_NULL_HANDLE;
 
 	DEC_DWORD_STAT(STAT_VulkanNumCmdBuffers);
@@ -200,8 +199,6 @@ void FVulkanCommandBuffer::Begin(FVulkanQueryPool* OptionalQueryPool, VkRenderPa
 		BindlessDescriptorManager->BindDescriptorBuffers(CommandBufferHandle, SupportedStages);
 	}
 
-	bNeedsDynamicStateSet = true;
-
 #if RHI_NEW_GPU_PROFILER
 	if (OptionalQueryPool)
 	{
@@ -286,7 +283,7 @@ FVulkanCommandBufferPool::FVulkanCommandBufferPool(FVulkanDevice& InDevice, FVul
 	ZeroVulkanStruct(CmdPoolInfo, VK_STRUCTURE_TYPE_COMMAND_POOL_CREATE_INFO);
 	CmdPoolInfo.queueFamilyIndex = InQueue.GetFamilyIndex();
 	CmdPoolInfo.flags = VK_COMMAND_POOL_CREATE_RESET_COMMAND_BUFFER_BIT; // :todo: Investigate use of VK_COMMAND_POOL_CREATE_TRANSIENT_BIT?
-	VERIFYVULKANRESULT(VulkanRHI::vkCreateCommandPool(Device.GetInstanceHandle(), &CmdPoolInfo, VULKAN_CPU_ALLOCATOR, &Handle));
+	VERIFYVULKANRESULT(VulkanRHI::vkCreateCommandPool(Device.GetHandle(), &CmdPoolInfo, VULKAN_CPU_ALLOCATOR, &Handle));
 }
 
 FVulkanCommandBufferPool::~FVulkanCommandBufferPool()
@@ -303,7 +300,7 @@ FVulkanCommandBufferPool::~FVulkanCommandBufferPool()
 		delete CmdBuffer;
 	}
 
-	VulkanRHI::vkDestroyCommandPool(Device.GetInstanceHandle(), Handle, VULKAN_CPU_ALLOCATOR);
+	VulkanRHI::vkDestroyCommandPool(Device.GetHandle(), Handle, VULKAN_CPU_ALLOCATOR);
 	Handle = VK_NULL_HANDLE;
 }
 
@@ -314,7 +311,7 @@ void FVulkanCommandBufferPool::FreeUnusedCmdBuffers(FVulkanQueue* InQueue, bool
 	
 	if (bTrimMemory)
 	{
-		VulkanRHI::vkTrimCommandPool(Device.GetInstanceHandle(), Handle, 0);
+		VulkanRHI::vkTrimCommandPool(Device.GetHandle(), Handle, 0);
 		return;
 	}
 
diff --git a/Engine/Source/Runtime/VulkanRHI/Private/VulkanCommandBuffer.h b/Engine/Source/Runtime/VulkanRHI/Private/VulkanCommandBuffer.h
--- a/Engine/Source/Runtime/VulkanRHI/Private/VulkanCommandBuffer.h
+++ b/Engine/Source/Runtime/VulkanRHI/Private/VulkanCommandBuffer.h
@@ -108,7 +108,6 @@ public:
 	TArray<VkRect2D, TInlineAllocator<2>> CurrentScissors;
 	uint32 CurrentStencilRef = 0;
 	EState State = EState::NotAllocated;
-	uint8 bNeedsDynamicStateSet			: 1 = 1;
 	uint8 bHasPipeline					: 1 = 0;
 	uint8 bHasViewport					: 1 = 0;
 	uint8 bHasScissor					: 1 = 0;
diff --git a/Engine/Source/Runtime/VulkanRHI/Private/VulkanDescriptorSets.cpp b/Engine/Source/Runtime/VulkanRHI/Private/VulkanDescriptorSets.cpp
--- a/Engine/Source/Runtime/VulkanRHI/Private/VulkanDescriptorSets.cpp
+++ b/Engine/Source/Runtime/VulkanRHI/Private/VulkanDescriptorSets.cpp
@@ -318,7 +318,7 @@ void FVulkanDescriptorSetsLayout::Compile(FVulkanDescriptorSetLayoutMap& DSetLay
 		DescriptorLayoutInfo.bindingCount = Layout.LayoutBindings.Num();
 		DescriptorLayoutInfo.pBindings = Layout.LayoutBindings.GetData();
 
-		VERIFYVULKANRESULT(VulkanRHI::vkCreateDescriptorSetLayout(Device->GetInstanceHandle(), &DescriptorLayoutInfo, VULKAN_CPU_ALLOCATOR, LayoutHandle));
+		VERIFYVULKANRESULT(VulkanRHI::vkCreateDescriptorSetLayout(Device->GetHandle(), &DescriptorLayoutInfo, VULKAN_CPU_ALLOCATOR, LayoutHandle));
 
 		if (LayoutHandleId)
 		{
diff --git a/Engine/Source/Runtime/VulkanRHI/Private/VulkanDevice.cpp b/Engine/Source/Runtime/VulkanRHI/Private/VulkanDevice.cpp
--- a/Engine/Source/Runtime/VulkanRHI/Private/VulkanDevice.cpp
+++ b/Engine/Source/Runtime/VulkanRHI/Private/VulkanDevice.cpp
@@ -1531,7 +1531,7 @@ void FVulkanDevice::NotifyDeletedImage(VkImage Image, bool bRenderTarget)
 	}
 }
 
-void FVulkanDevice::NotifyDeletedGfxPipeline(class FVulkanRHIGraphicsPipelineState* Pipeline)
+void FVulkanDevice::NotifyDeletedGfxPipeline(class FVulkanGraphicsPipelineState* Pipeline)
 {
 	//#todo-rco: Loop through all contexts!
 	if (ImmediateContext && ImmediateContext->PendingGfxState)
@@ -1599,7 +1599,7 @@ void FVulkanDevice::GetCalibrationTimestamp(FVulkanTiming& InOutTiming)
 		const bool bHasHostTimeDomain = (CalibratedTimestampInfo[1].timeDomain != VK_TIME_DOMAIN_DEVICE_KHR);
 		const uint32 TimestampCount = bHasHostTimeDomain ? 2 : 1;
 
-		VERIFYVULKANRESULT(VulkanDynamicAPI::vkGetCalibratedTimestampsEXT(GetInstanceHandle(), TimestampCount, CalibratedTimestampInfo, Timestamps, MaxDeviations));
+		VERIFYVULKANRESULT(VulkanDynamicAPI::vkGetCalibratedTimestampsEXT(GetHandle(), TimestampCount, CalibratedTimestampInfo, Timestamps, MaxDeviations));
 
 		InOutTiming.GPUTimestamp = Timestamps[0];
 		switch (CalibratedTimestampInfo[1].timeDomain)
diff --git a/Engine/Source/Runtime/VulkanRHI/Private/VulkanDevice.h b/Engine/Source/Runtime/VulkanRHI/Private/VulkanDevice.h
--- a/Engine/Source/Runtime/VulkanRHI/Private/VulkanDevice.h
+++ b/Engine/Source/Runtime/VulkanRHI/Private/VulkanDevice.h
@@ -406,7 +406,7 @@ public:
 
 	const VkComponentMapping& GetFormatComponentMapping(EPixelFormat UEFormat) const;
 
-	VkDevice GetInstanceHandle() const
+	VkDevice GetHandle() const
 	{
 		return Device;
 	}
@@ -517,7 +517,7 @@ public:
 		return PipelineStateCache;
 	}
 
-	void NotifyDeletedGfxPipeline(class FVulkanRHIGraphicsPipelineState* Pipeline);
+	void NotifyDeletedGfxPipeline(class FVulkanGraphicsPipelineState* Pipeline);
 	void NotifyDeletedComputePipeline(class FVulkanComputePipeline* Pipeline);
 
 	void VulkanSetObjectName(VkObjectType Type, uint64_t Handle, const TCHAR* Name);
@@ -683,5 +683,5 @@ private:
 
 	class FVulkanPipelineStateCacheManager* PipelineStateCache;
 	friend class FVulkanDynamicRHI;
-	friend class FVulkanRHIGraphicsPipelineState;
+	friend class FVulkanGraphicsPipelineState;
 };
diff --git a/Engine/Source/Runtime/VulkanRHI/Private/VulkanFramebuffer.cpp b/Engine/Source/Runtime/VulkanRHI/Private/VulkanFramebuffer.cpp
--- a/Engine/Source/Runtime/VulkanRHI/Private/VulkanFramebuffer.cpp
+++ b/Engine/Source/Runtime/VulkanRHI/Private/VulkanFramebuffer.cpp
@@ -107,7 +107,7 @@ FVulkanFramebuffer::FVulkanFramebuffer(FVulkanDevice& Device, const FRHISetRende
 	CreateInfo.height = RTExtents.height;
 	CreateInfo.layers = NumLayers;
 
-	VERIFYVULKANRESULT_EXPANDED(VulkanRHI::vkCreateFramebuffer(Device.GetInstanceHandle(), &CreateInfo, VULKAN_CPU_ALLOCATOR, &Framebuffer));
+	VERIFYVULKANRESULT_EXPANDED(VulkanRHI::vkCreateFramebuffer(Device.GetHandle(), &CreateInfo, VULKAN_CPU_ALLOCATOR, &Framebuffer));
 
 	RenderArea.offset.x = 0;
 	RenderArea.offset.y = 0;
diff --git a/Engine/Source/Runtime/VulkanRHI/Private/VulkanLayout.cpp b/Engine/Source/Runtime/VulkanRHI/Private/VulkanLayout.cpp
--- a/Engine/Source/Runtime/VulkanRHI/Private/VulkanLayout.cpp
+++ b/Engine/Source/Runtime/VulkanRHI/Private/VulkanLayout.cpp
@@ -41,6 +41,6 @@ void FVulkanLayout::Compile(FVulkanDescriptorSetLayoutMap& DSetLayoutMap)
 		PipelineLayoutCreateInfo.setLayoutCount = LayoutHandles.Num();
 		PipelineLayoutCreateInfo.pSetLayouts = LayoutHandles.GetData();
 		//PipelineLayoutCreateInfo.pushConstantRangeCount = 0;
-		VERIFYVULKANRESULT(VulkanRHI::vkCreatePipelineLayout(Device->GetInstanceHandle(), &PipelineLayoutCreateInfo, VULKAN_CPU_ALLOCATOR, &PipelineLayout));
+		VERIFYVULKANRESULT(VulkanRHI::vkCreatePipelineLayout(Device->GetHandle(), &PipelineLayoutCreateInfo, VULKAN_CPU_ALLOCATOR, &PipelineLayout));
 	}
 }
diff --git a/Engine/Source/Runtime/VulkanRHI/Private/VulkanMemory.cpp b/Engine/Source/Runtime/VulkanRHI/Private/VulkanMemory.cpp
--- a/Engine/Source/Runtime/VulkanRHI/Private/VulkanMemory.cpp
+++ b/Engine/Source/Runtime/VulkanRHI/Private/VulkanMemory.cpp
@@ -606,7 +606,7 @@ namespace VulkanRHI
 
 		bHasUnifiedMemory = FVulkanPlatform::HasUnifiedMemory();
 
-		DeviceHandle = Device->GetInstanceHandle();
+		DeviceHandle = Device->GetHandle();
 		UpdateMemoryProperties();
 
 		PrimaryHeapIndex = -1;
@@ -2506,7 +2506,7 @@ namespace VulkanRHI
 		// Does not need to go in the deferred deletion queue
 		if(Buffer != VK_NULL_HANDLE)
 		{
-			VulkanRHI::vkDestroyBuffer(Device->GetInstanceHandle(), Buffer, VULKAN_CPU_ALLOCATOR);
+			VulkanRHI::vkDestroyBuffer(Device->GetHandle(), Buffer, VULKAN_CPU_ALLOCATOR);
 			Buffer = VK_NULL_HANDLE;
 		}
 	}
@@ -2782,7 +2782,7 @@ namespace VulkanRHI
 		VkBuffer Buffer = Device->CreateBuffer(BufferSize, BufferUsageFlags);
 
 		VkMemoryRequirements MemReqs;
-		VulkanRHI::vkGetBufferMemoryRequirements(Device->GetInstanceHandle(), Buffer, &MemReqs);
+		VulkanRHI::vkGetBufferMemoryRequirements(Device->GetHandle(), Buffer, &MemReqs);
 		Alignment = FMath::Max((uint32)MemReqs.alignment, Alignment);
 		ensure(MemReqs.size >= BufferSize);
 
@@ -2805,7 +2805,7 @@ namespace VulkanRHI
 			HandleOOM(false);
 			checkNoEntry();
 		}
-		VERIFYVULKANRESULT(VulkanRHI::vkBindBufferMemory(Device->GetInstanceHandle(), Buffer, DeviceMemoryAllocation->GetHandle(), 0));
+		VERIFYVULKANRESULT(VulkanRHI::vkBindBufferMemory(Device->GetHandle(), Buffer, DeviceMemoryAllocation->GetHandle(), 0));
 		uint8 AllocationFlags = 0;
 		if(!bHasUnifiedMemory && MetaTypeCanEvict(MetaType))
 		{
@@ -3048,7 +3048,7 @@ namespace VulkanRHI
 		ZeroVulkanStruct(MemoryRequirements, VK_STRUCTURE_TYPE_MEMORY_REQUIREMENTS_2);
 		MemoryRequirements.pNext = &DedicatedRequirements;
 
-		VulkanRHI::vkGetBufferMemoryRequirements2(Device->GetInstanceHandle(), &BufferMemoryRequirementsInfo, &MemoryRequirements);
+		VulkanRHI::vkGetBufferMemoryRequirements2(Device->GetHandle(), &BufferMemoryRequirementsInfo, &MemoryRequirements);
 
 		// Allow caller to force his own alignment requirements
 		MemoryRequirements.memoryRequirements.alignment = FMath::Max(MemoryRequirements.memoryRequirements.alignment, (VkDeviceSize)InForceMinAlignment);
@@ -3074,7 +3074,7 @@ namespace VulkanRHI
 				BindBufferMemoryInfo.memory = OutAllocation.GetDeviceMemoryHandle(Device);
 				BindBufferMemoryInfo.memoryOffset = OutAllocation.Offset;
 
-				VERIFYVULKANRESULT(VulkanRHI::vkBindBufferMemory2(Device->GetInstanceHandle(), 1, &BindBufferMemoryInfo));
+				VERIFYVULKANRESULT(VulkanRHI::vkBindBufferMemory2(Device->GetHandle(), 1, &BindBufferMemoryInfo));
 			}
 
 			if (InDebugName)
@@ -3109,7 +3109,7 @@ namespace VulkanRHI
 		ZeroVulkanStruct(MemoryReqs2, VK_STRUCTURE_TYPE_MEMORY_REQUIREMENTS_2_KHR);
 		MemoryReqs2.pNext = &DedMemoryReqs;
 
-		VulkanRHI::vkGetImageMemoryRequirements2(Device->GetInstanceHandle(), &ImageMemoryReqs2, &MemoryReqs2);
+		VulkanRHI::vkGetImageMemoryRequirements2(Device->GetHandle(), &ImageMemoryReqs2, &MemoryReqs2);
 
 		bool bUseDedicated = DedMemoryReqs.prefersDedicatedAllocation != VK_FALSE || DedMemoryReqs.requiresDedicatedAllocation != VK_FALSE;
 		if (bUseDedicated)
@@ -3881,7 +3881,7 @@ namespace VulkanRHI
 	void FVulkanAllocation::BindBuffer(FVulkanDevice* Device, VkBuffer Buffer)
 	{
 		VkDeviceMemory MemoryHandle = GetDeviceMemoryHandle(Device);
-		VkResult Result = VulkanRHI::vkBindBufferMemory(Device->GetInstanceHandle(), Buffer, MemoryHandle, Offset);
+		VkResult Result = VulkanRHI::vkBindBufferMemory(Device->GetHandle(), Buffer, MemoryHandle, Offset);
 		if (Result == VK_ERROR_OUT_OF_DEVICE_MEMORY || Result == VK_ERROR_OUT_OF_HOST_MEMORY)
 		{
 			Device->GetMemoryManager().DumpMemory();
@@ -3891,7 +3891,7 @@ namespace VulkanRHI
 	void FVulkanAllocation::BindImage(FVulkanDevice* Device, VkImage Image)
 	{
 		VkDeviceMemory MemoryHandle = GetDeviceMemoryHandle(Device);
-		VkResult Result = VulkanRHI::vkBindImageMemory(Device->GetInstanceHandle(), Image, MemoryHandle, Offset);
+		VkResult Result = VulkanRHI::vkBindImageMemory(Device->GetHandle(), Image, MemoryHandle, Offset);
 		if (Result == VK_ERROR_OUT_OF_DEVICE_MEMORY || Result == VK_ERROR_OUT_OF_HOST_MEMORY)
 		{
 			Device->GetMemoryManager().DumpMemory();
@@ -4296,7 +4296,7 @@ namespace VulkanRHI
 	void FStagingBuffer::Destroy()
 	{
 		//// Does not need to go in the deferred deletion queue
-		VulkanRHI::vkDestroyBuffer(Device->GetInstanceHandle(), Buffer, VULKAN_CPU_ALLOCATOR);
+		VulkanRHI::vkDestroyBuffer(Device->GetHandle(), Buffer, VULKAN_CPU_ALLOCATOR);
 		Buffer = VK_NULL_HANDLE;
 		Device->GetMemoryManager().FreeVulkanAllocation(Allocation);
 	}
@@ -4527,7 +4527,7 @@ namespace VulkanRHI
 		VkFenceCreateInfo Info;
 		ZeroVulkanStruct(Info, VK_STRUCTURE_TYPE_FENCE_CREATE_INFO);
 		Info.flags = bCreateSignaled ? VK_FENCE_CREATE_SIGNALED_BIT : 0;
-		VERIFYVULKANRESULT(VulkanRHI::vkCreateFence(InDevice->GetInstanceHandle(), &Info, VULKAN_CPU_ALLOCATOR, &Handle));
+		VERIFYVULKANRESULT(VulkanRHI::vkCreateFence(InDevice->GetHandle(), &Info, VULKAN_CPU_ALLOCATOR, &Handle));
 	}
 
 	FFence::~FFence()
@@ -4543,7 +4543,7 @@ namespace VulkanRHI
 	inline void FFenceManager::DestroyFence(FFence* Fence)
 	{
 		// Does not need to go in the deferred deletion queue
-		VulkanRHI::vkDestroyFence(Device->GetInstanceHandle(), Fence->GetHandle(), VULKAN_CPU_ALLOCATOR);
+		VulkanRHI::vkDestroyFence(Device->GetHandle(), Fence->GetHandle(), VULKAN_CPU_ALLOCATOR);
 		Fence->Handle = VK_NULL_HANDLE;
 		delete Fence;
 	}
@@ -4557,7 +4557,7 @@ namespace VulkanRHI
 	{
 		FScopeLock Lock(&FenceLock);
 		ensureMsgf(UsedFences.Num() == 0, TEXT("No all fences are done!"));
-		VkDevice DeviceHandle = Device->GetInstanceHandle();
+		VkDevice DeviceHandle = Device->GetHandle();
 		for (FFence* Fence : FreeFences)
 		{
 			DestroyFence(Fence);
@@ -4617,7 +4617,7 @@ namespace VulkanRHI
 	{
 		check(UsedFences.Contains(Fence));
 		check(Fence->State == FFence::EState::NotReady);
-		VkResult Result = VulkanRHI::vkGetFenceStatus(Device->GetInstanceHandle(), Fence->Handle);
+		VkResult Result = VulkanRHI::vkGetFenceStatus(Device->GetHandle(), Fence->Handle);
 		switch (Result)
 		{
 		case VK_SUCCESS:
@@ -4643,7 +4643,7 @@ namespace VulkanRHI
 
 		check(UsedFences.Contains(Fence));
 		check(Fence->State == FFence::EState::NotReady);
-		VkResult Result = VulkanRHI::vkWaitForFences(Device->GetInstanceHandle(), 1, &Fence->Handle, true, TimeInNanoseconds);
+		VkResult Result = VulkanRHI::vkWaitForFences(Device->GetHandle(), 1, &Fence->Handle, true, TimeInNanoseconds);
 		switch (Result)
 		{
 		case VK_SUCCESS:
@@ -4672,7 +4672,7 @@ namespace VulkanRHI
 			}
 		}
 
-		const VkResult Result = VulkanRHI::vkWaitForFences(Device->GetInstanceHandle(), FenceHandles.Num(), FenceHandles.GetData(), false, TimeInNanoseconds);
+		const VkResult Result = VulkanRHI::vkWaitForFences(Device->GetHandle(), FenceHandles.Num(), FenceHandles.GetData(), false, TimeInNanoseconds);
 		switch (Result)
 		{
 		case VK_SUCCESS:
@@ -4695,7 +4695,7 @@ namespace VulkanRHI
 	{
 		if (Fence->State != FFence::EState::NotReady)
 		{
-			VERIFYVULKANRESULT(VulkanRHI::vkResetFences(Device->GetInstanceHandle(), 1, &Fence->Handle));
+			VERIFYVULKANRESULT(VulkanRHI::vkResetFences(Device->GetHandle(), 1, &Fence->Handle));
 			Fence->State = FFence::EState::NotReady;
 		}
 	}
@@ -4706,7 +4706,7 @@ namespace VulkanRHI
 	{
 		VkEventCreateInfo Info;
 		ZeroVulkanStruct(Info, VK_STRUCTURE_TYPE_EVENT_CREATE_INFO);
-		VERIFYVULKANRESULT(VulkanRHI::vkCreateEvent(InDevice->GetInstanceHandle(), &Info, VULKAN_CPU_ALLOCATOR, &Handle));
+		VERIFYVULKANRESULT(VulkanRHI::vkCreateEvent(InDevice->GetHandle(), &Info, VULKAN_CPU_ALLOCATOR, &Handle));
 	}
 
 	FGPUEvent::~FGPUEvent()
@@ -4841,7 +4841,7 @@ namespace VulkanRHI
 
 	void FDeferredDeletionQueue2::ReleaseResourcesImmediately(const TArray<FEntry>& InEntries)
 	{
-		VkDevice DeviceHandle = Device->GetInstanceHandle();
+		VkDevice DeviceHandle = Device->GetHandle();
 
 		// Traverse list backwards so the swap switches to elements already tested
 		for (const FEntry& Entry : InEntries)
@@ -4912,7 +4912,7 @@ namespace VulkanRHI
 	{
 		auto FreeBlock = [](FVulkanDevice* InDevice, FTempMemoryBlock* Block)
 		{
-			VulkanRHI::vkDestroyBuffer(InDevice->GetInstanceHandle(), Block->Buffer, VULKAN_CPU_ALLOCATOR);
+			VulkanRHI::vkDestroyBuffer(InDevice->GetHandle(), Block->Buffer, VULKAN_CPU_ALLOCATOR);
 			InDevice->GetMemoryManager().FreeVulkanAllocation(Block->Allocation, VulkanRHI::EVulkanFreeFlag_DontDefer);
 			delete Block;
 		};
@@ -4964,7 +4964,7 @@ namespace VulkanRHI
 			VkBufferDeviceAddressInfo BufferInfo;
 			ZeroVulkanStruct(BufferInfo, VK_STRUCTURE_TYPE_BUFFER_DEVICE_ADDRESS_INFO);
 			BufferInfo.buffer = NewBlock->Buffer;
-			NewBlock->BufferAddress = VulkanRHI::vkGetBufferDeviceAddressKHR(Device->GetInstanceHandle(), &BufferInfo);
+			NewBlock->BufferAddress = VulkanRHI::vkGetBufferDeviceAddressKHR(Device->GetHandle(), &BufferInfo);
 		}
 
 		return NewBlock;
@@ -5117,7 +5117,7 @@ namespace VulkanRHI
 			CreateInfo.pNext = &TypeCreateInfo;
 		}
 
-		VERIFYVULKANRESULT(VulkanRHI::vkCreateSemaphore(Device.GetInstanceHandle(), &CreateInfo, VULKAN_CPU_ALLOCATOR, &SemaphoreHandle));
+		VERIFYVULKANRESULT(VulkanRHI::vkCreateSemaphore(Device.GetHandle(), &CreateInfo, VULKAN_CPU_ALLOCATOR, &SemaphoreHandle));
 	}
 
 	FSemaphore::FSemaphore(FVulkanDevice& InDevice, const VkSemaphore& InExternalSemaphore)
@@ -5131,7 +5131,7 @@ namespace VulkanRHI
 		check(SemaphoreHandle != VK_NULL_HANDLE);
 		if (EnumHasAnyFlags(Flags, EVulkanSemaphoreFlags::ImmediateDeletion))
 		{
-			VulkanRHI::vkDestroySemaphore(Device.GetInstanceHandle(), SemaphoreHandle, VULKAN_CPU_ALLOCATOR);
+			VulkanRHI::vkDestroySemaphore(Device.GetHandle(), SemaphoreHandle, VULKAN_CPU_ALLOCATOR);
 		}
 		else if (!IsExternallyOwned())
 		{
@@ -5149,7 +5149,7 @@ namespace VulkanRHI
 		WaitInfo.semaphoreCount = 1;
 		WaitInfo.pSemaphores = &SemaphoreHandle;
 		WaitInfo.pValues = (uint64_t*)&Value;
-		const VkResult Result = VulkanRHI::vkWaitSemaphoresKHR(Device.GetInstanceHandle(), &WaitInfo, Timeout);
+		const VkResult Result = VulkanRHI::vkWaitSemaphoresKHR(Device.GetHandle(), &WaitInfo, Timeout);
 		switch (Result)
 		{
 		case VK_SUCCESS:
@@ -5167,7 +5167,7 @@ namespace VulkanRHI
 	{
 		checkf(IsTimeline(), TEXT("GetTimelineSemaphoreValue not available!  Semaphore was not created with EVulkanSemaphoreFlags::Timeline."));
 		uint64 Value = 0;
-		VERIFYVULKANRESULT(VulkanRHI::vkGetSemaphoreCounterValueKHR(Device.GetInstanceHandle(), SemaphoreHandle, (uint64_t*)&Value));
+		VERIFYVULKANRESULT(VulkanRHI::vkGetSemaphoreCounterValueKHR(Device.GetHandle(), SemaphoreHandle, (uint64_t*)&Value));
 		return Value;
 	}
 }
diff --git a/Engine/Source/Runtime/VulkanRHI/Private/VulkanPendingState.cpp b/Engine/Source/Runtime/VulkanRHI/Private/VulkanPendingState.cpp
--- a/Engine/Source/Runtime/VulkanRHI/Private/VulkanPendingState.cpp
+++ b/Engine/Source/Runtime/VulkanRHI/Private/VulkanPendingState.cpp
@@ -62,7 +62,7 @@ FVulkanDescriptorPool::FVulkanDescriptorPool(FVulkanDevice* InDevice, const FVul
 #if VULKAN_ENABLE_AGGRESSIVE_STATS
 	SCOPE_CYCLE_COUNTER(STAT_VulkanVkCreateDescriptorPool);
 #endif
-	VERIFYVULKANRESULT(VulkanRHI::vkCreateDescriptorPool(Device->GetInstanceHandle(), &PoolInfo, VULKAN_CPU_ALLOCATOR, &DescriptorPool));
+	VERIFYVULKANRESULT(VulkanRHI::vkCreateDescriptorPool(Device->GetHandle(), &PoolInfo, VULKAN_CPU_ALLOCATOR, &DescriptorPool));
 
 	INC_DWORD_STAT_BY(STAT_VulkanNumDescSetsTotal, MaxDescriptorSets);
 }
@@ -74,7 +74,7 @@ FVulkanDescriptorPool::~FVulkanDescriptorPool()
 
 	if (DescriptorPool != VK_NULL_HANDLE)
 	{
-		VulkanRHI::vkDestroyDescriptorPool(Device->GetInstanceHandle(), DescriptorPool, VULKAN_CPU_ALLOCATOR);
+		VulkanRHI::vkDestroyDescriptorPool(Device->GetHandle(), DescriptorPool, VULKAN_CPU_ALLOCATOR);
 		DescriptorPool = VK_NULL_HANDLE;
 	}
 }
@@ -113,7 +113,7 @@ void FVulkanDescriptorPool::Reset()
 {
 	if (DescriptorPool != VK_NULL_HANDLE)
 	{
-		VERIFYVULKANRESULT(VulkanRHI::vkResetDescriptorPool(Device->GetInstanceHandle(), DescriptorPool, 0));
+		VERIFYVULKANRESULT(VulkanRHI::vkResetDescriptorPool(Device->GetHandle(), DescriptorPool, 0));
 	}
 
 	NumAllocatedDescriptorSets = 0;
@@ -124,7 +124,7 @@ bool FVulkanDescriptorPool::AllocateDescriptorSets(const VkDescriptorSetAllocate
 	VkDescriptorSetAllocateInfo DescriptorSetAllocateInfo = InDescriptorSetAllocateInfo;
 	DescriptorSetAllocateInfo.descriptorPool = DescriptorPool;
 
-	return VK_SUCCESS == VulkanRHI::vkAllocateDescriptorSets(Device->GetInstanceHandle(), &DescriptorSetAllocateInfo, OutSets);
+	return VK_SUCCESS == VulkanRHI::vkAllocateDescriptorSets(Device->GetHandle(), &DescriptorSetAllocateInfo, OutSets);
 }
 
 FVulkanTypedDescriptorPoolSet::~FVulkanTypedDescriptorPoolSet()
@@ -373,7 +373,7 @@ void FVulkanPendingComputeState::PrepareForDispatch(FVulkanCommandListContext& C
 
 FVulkanPendingGfxState::~FVulkanPendingGfxState()
 {
-	TMap<FVulkanRHIGraphicsPipelineState*, FVulkanGraphicsPipelineDescriptorState*> Temp;
+	TMap<FVulkanGraphicsPipelineState*, FVulkanGraphicsPipelineDescriptorState*> Temp;
 	Swap(Temp, PipelineStates);
 	for (auto& Pair : Temp)
 	{
@@ -530,8 +530,6 @@ void FVulkanPendingGfxState::InternalUpdateDynamicStates(FVulkanCommandBuffer& C
 		CommandBuffer.CurrentStencilRef = StencilRef;
 		CommandBuffer.bHasStencilRef = true;
 	}
-
-	CommandBuffer.bNeedsDynamicStateSet = false;
 }
 
 void FVulkanPendingGfxState::UpdateInputAttachments(FVulkanFramebuffer* Framebuffer)
@@ -643,7 +641,7 @@ FVulkanGenericDescriptorPool::FVulkanGenericDescriptorPool(FVulkanDevice* InDevi
 #if VULKAN_ENABLE_AGGRESSIVE_STATS
 	SCOPE_CYCLE_COUNTER(STAT_VulkanVkCreateDescriptorPool);
 #endif
-	VERIFYVULKANRESULT(VulkanRHI::vkCreateDescriptorPool(Device->GetInstanceHandle(), &PoolInfo, VULKAN_CPU_ALLOCATOR, &DescriptorPool));
+	VERIFYVULKANRESULT(VulkanRHI::vkCreateDescriptorPool(Device->GetHandle(), &PoolInfo, VULKAN_CPU_ALLOCATOR, &DescriptorPool));
 
 	INC_DWORD_STAT_BY(STAT_VulkanNumDescSetsTotal, MaxDescriptorSets);
 	INC_DWORD_STAT(STAT_VulkanNumDescPools);
@@ -656,14 +654,14 @@ FVulkanGenericDescriptorPool::~FVulkanGenericDescriptorPool()
 		DEC_DWORD_STAT_BY(STAT_VulkanNumDescSetsTotal, MaxDescriptorSets);
 		DEC_DWORD_STAT(STAT_VulkanNumDescPools);
 
-		VulkanRHI::vkDestroyDescriptorPool(Device->GetInstanceHandle(), DescriptorPool, VULKAN_CPU_ALLOCATOR);
+		VulkanRHI::vkDestroyDescriptorPool(Device->GetHandle(), DescriptorPool, VULKAN_CPU_ALLOCATOR);
 	}
 }
 
 void FVulkanGenericDescriptorPool::Reset()
 {
 	check(DescriptorPool != VK_NULL_HANDLE);
-	VERIFYVULKANRESULT(VulkanRHI::vkResetDescriptorPool(Device->GetInstanceHandle(), DescriptorPool, 0));
+	VERIFYVULKANRESULT(VulkanRHI::vkResetDescriptorPool(Device->GetHandle(), DescriptorPool, 0));
 }
 
 bool FVulkanGenericDescriptorPool::AllocateDescriptorSet(VkDescriptorSetLayout Layout, VkDescriptorSet& OutSet)
@@ -676,7 +674,7 @@ bool FVulkanGenericDescriptorPool::AllocateDescriptorSet(VkDescriptorSetLayout L
 	DescriptorSetAllocateInfo.descriptorSetCount = 1;
 	DescriptorSetAllocateInfo.pSetLayouts = &Layout;
 
-	return (VK_SUCCESS == VulkanRHI::vkAllocateDescriptorSets(Device->GetInstanceHandle(), &DescriptorSetAllocateInfo, &OutSet));
+	return (VK_SUCCESS == VulkanRHI::vkAllocateDescriptorSets(Device->GetHandle(), &DescriptorSetAllocateInfo, &OutSet));
 }
 
 FVulkanDescriptorSetCache::FVulkanDescriptorSetCache(FVulkanDevice* InDevice)
@@ -835,7 +833,7 @@ bool FVulkanDescriptorSetCache::FCachedPool::CreateDescriptorSets(
 		INC_DWORD_STAT(STAT_VulkanNumDescSets);
 		SCOPE_CYCLE_COUNTER(STAT_VulkanVkUpdateDS);
 #endif
-		VulkanRHI::vkUpdateDescriptorSets(Pool.GetDevice()->GetInstanceHandle(),
+		VulkanRHI::vkUpdateDescriptorSets(Pool.GetDevice()->GetHandle(),
 			DSWriter.GetNumWrites(), DSWriter.GetWriteDescriptors(), 0, nullptr);
 	}
 
@@ -881,7 +879,7 @@ void FVulkanDescriptorSetCache::FCachedPool::CalcPoolSizesRatio(float PoolSizesR
 	}
 }
 
-bool FVulkanPendingGfxState::SetGfxPipeline(FVulkanRHIGraphicsPipelineState* InGfxPipeline, bool bForceReset)
+bool FVulkanPendingGfxState::SetGfxPipeline(FVulkanGraphicsPipelineState* InGfxPipeline, bool bForceReset)
 {
 	bool bChanged = bForceReset;
 
diff --git a/Engine/Source/Runtime/VulkanRHI/Private/VulkanPendingState.h b/Engine/Source/Runtime/VulkanRHI/Private/VulkanPendingState.h
--- a/Engine/Source/Runtime/VulkanRHI/Private/VulkanPendingState.h
+++ b/Engine/Source/Runtime/VulkanRHI/Private/VulkanPendingState.h
@@ -283,7 +283,7 @@ public:
 
 	void PrepareForDraw(FVulkanCommandListContext& Context);
 
-	bool SetGfxPipeline(FVulkanRHIGraphicsPipelineState* InGfxPipeline, bool bForceReset);
+	bool SetGfxPipeline(FVulkanGraphicsPipelineState* InGfxPipeline, bool bForceReset);
 
 	inline void UpdateDynamicStates(FVulkanCommandBuffer& CommandBuffer)
 	{
@@ -298,7 +298,7 @@ public:
 		}
 	}
 
-	void NotifyDeletedPipeline(FVulkanRHIGraphicsPipelineState* Pipeline)
+	void NotifyDeletedPipeline(FVulkanGraphicsPipelineState* Pipeline)
 	{
 		PipelineStates.Remove(Pipeline);
 	}
@@ -317,10 +317,10 @@ protected:
 
 	bool bNeedToClear;
 
-	FVulkanRHIGraphicsPipelineState* CurrentPipeline;
+	FVulkanGraphicsPipelineState* CurrentPipeline;
 	FVulkanGraphicsPipelineDescriptorState* CurrentState;
 
-	TMap<FVulkanRHIGraphicsPipelineState*, FVulkanGraphicsPipelineDescriptorState*> PipelineStates;
+	TMap<FVulkanGraphicsPipelineState*, FVulkanGraphicsPipelineDescriptorState*> PipelineStates;
 
 	struct FVertexStream
 	{
diff --git a/Engine/Source/Runtime/VulkanRHI/Private/VulkanPipeline.cpp b/Engine/Source/Runtime/VulkanRHI/Private/VulkanPipeline.cpp
--- a/Engine/Source/Runtime/VulkanRHI/Private/VulkanPipeline.cpp
+++ b/Engine/Source/Runtime/VulkanRHI/Private/VulkanPipeline.cpp
@@ -237,7 +237,7 @@ FVulkanComputePipeline::~FVulkanComputePipeline()
 	DEC_DWORD_STAT(STAT_VulkanNumComputePSOs);
 }
 
-FVulkanRHIGraphicsPipelineState::~FVulkanRHIGraphicsPipelineState()
+FVulkanGraphicsPipelineState::~FVulkanGraphicsPipelineState()
 {
 #if !UE_BUILD_SHIPPING
 	SGraphicsRHICount--;
@@ -255,7 +255,7 @@ FVulkanRHIGraphicsPipelineState::~FVulkanRHIGraphicsPipelineState()
 	}
 }
 
-void FVulkanRHIGraphicsPipelineState::GetOrCreateShaderModules(TRefCountPtr<FVulkanShaderModule> (&ShaderModulesOUT)[ShaderStage::NumGraphicsStages], FVulkanShader*const* Shaders)
+void FVulkanGraphicsPipelineState::GetOrCreateShaderModules(TRefCountPtr<FVulkanShaderModule> (&ShaderModulesOUT)[ShaderStage::NumGraphicsStages], FVulkanShader*const* Shaders)
 {
 	for (int32 Index = 0; Index < ShaderStage::NumGraphicsStages; ++Index)
 	{
@@ -268,13 +268,13 @@ void FVulkanRHIGraphicsPipelineState::GetOrCreateShaderModules(TRefCountPtr<FVul
 	}
 }
 
-FVulkanShader::FSpirvCode FVulkanRHIGraphicsPipelineState::GetPatchedSpirvCode(FVulkanShader* Shader)
+FVulkanShader::FSpirvCode FVulkanGraphicsPipelineState::GetPatchedSpirvCode(FVulkanShader* Shader)
 {
 	check(Shader);
 	return Shader->GetPatchedSpirvCode(Desc, Layout);
 }
 
-void FVulkanRHIGraphicsPipelineState::PurgeShaderModules(FVulkanShader*const* Shaders)
+void FVulkanGraphicsPipelineState::PurgeShaderModules(FVulkanShader*const* Shaders)
 {
 	for (int32 Index = 0; Index < ShaderStage::NumGraphicsStages; ++Index)
 	{
@@ -316,13 +316,13 @@ FVulkanPipelineStateCacheManager::~FVulkanPipelineStateCacheManager()
 	}
 	for (auto& Pair : DSetLayoutMap)
 	{
-		VulkanRHI::vkDestroyDescriptorSetLayout(Device->GetInstanceHandle(), Pair.Value.Handle, VULKAN_CPU_ALLOCATOR);
+		VulkanRHI::vkDestroyDescriptorSetLayout(Device->GetHandle(), Pair.Value.Handle, VULKAN_CPU_ALLOCATOR);
 	}
 
 	{
 		//TODO: Save PSOCache here?!
 		FScopedPipelineCache PipelineCacheExclusive = GlobalPSOCache.Get(EPipelineCacheAccess::Exclusive);
-		VulkanRHI::vkDestroyPipelineCache(Device->GetInstanceHandle(), PipelineCacheExclusive.Get(), VULKAN_CPU_ALLOCATOR);
+		VulkanRHI::vkDestroyPipelineCache(Device->GetHandle(), PipelineCacheExclusive.Get(), VULKAN_CPU_ALLOCATOR);
 	}
 
 	{
@@ -330,7 +330,7 @@ FVulkanPipelineStateCacheManager::~FVulkanPipelineStateCacheManager()
 		if (PipelineCacheExclusive.Get() != VK_NULL_HANDLE)
 		{
 			//If CurrentOpenedPSOCache is still valid then it has never received OnShaderPipelineCachePrecompilationComplete callback, so we are not going to save it's content to disk as it's most likely is incomplete at this point
-			VulkanRHI::vkDestroyPipelineCache(Device->GetInstanceHandle(), PipelineCacheExclusive.Get(), VULKAN_CPU_ALLOCATOR);
+			VulkanRHI::vkDestroyPipelineCache(Device->GetHandle(), PipelineCacheExclusive.Get(), VULKAN_CPU_ALLOCATOR);
 		}
 	}
 }
@@ -358,16 +358,16 @@ bool FVulkanPipelineStateCacheManager::Load(const TArray<FString>& CacheFilename
 
 				if (PipelineCacheExclusive.Get() == VK_NULL_HANDLE)
 				{
-					VERIFYVULKANRESULT(VulkanRHI::vkCreatePipelineCache(Device->GetInstanceHandle(), &PipelineCacheInfo, VULKAN_CPU_ALLOCATOR, &PipelineCacheExclusive.Get()));
+					VERIFYVULKANRESULT(VulkanRHI::vkCreatePipelineCache(Device->GetHandle(), &PipelineCacheInfo, VULKAN_CPU_ALLOCATOR, &PipelineCacheExclusive.Get()));
 				}
 				else
 				{
 					//TODO: assert on reopening the same cache twice?!
 					// if we have one already, create a temp one and merge it
 					VkPipelineCache TempPipelineCache;
-					VERIFYVULKANRESULT(VulkanRHI::vkCreatePipelineCache(Device->GetInstanceHandle(), &PipelineCacheInfo, VULKAN_CPU_ALLOCATOR, &TempPipelineCache));
-					VERIFYVULKANRESULT(VulkanRHI::vkMergePipelineCaches(Device->GetInstanceHandle(), PipelineCacheExclusive.Get(), 1, &TempPipelineCache));
-					VulkanRHI::vkDestroyPipelineCache(Device->GetInstanceHandle(), TempPipelineCache, VULKAN_CPU_ALLOCATOR);
+					VERIFYVULKANRESULT(VulkanRHI::vkCreatePipelineCache(Device->GetHandle(), &PipelineCacheInfo, VULKAN_CPU_ALLOCATOR, &TempPipelineCache));
+					VERIFYVULKANRESULT(VulkanRHI::vkMergePipelineCaches(Device->GetHandle(), PipelineCacheExclusive.Get(), 1, &TempPipelineCache));
+					VulkanRHI::vkDestroyPipelineCache(Device->GetHandle(), TempPipelineCache, VULKAN_CPU_ALLOCATOR);
 				}
 
 				double EndTime = FPlatformTime::Seconds();
@@ -427,7 +427,7 @@ bool FVulkanPipelineStateCacheManager::Load(const TArray<FString>& CacheFilename
 	{
 		VkPipelineCacheCreateInfo PipelineCacheInfo;
 		ZeroVulkanStruct(PipelineCacheInfo, VK_STRUCTURE_TYPE_PIPELINE_CACHE_CREATE_INFO);
-		VERIFYVULKANRESULT(VulkanRHI::vkCreatePipelineCache(Device->GetInstanceHandle(), &PipelineCacheInfo, VULKAN_CPU_ALLOCATOR, &PipelineCacheExclusive.Get()));
+		VERIFYVULKANRESULT(VulkanRHI::vkCreatePipelineCache(Device->GetHandle(), &PipelineCacheInfo, VULKAN_CPU_ALLOCATOR, &PipelineCacheExclusive.Get()));
 	}
 
 	return bResult;
@@ -469,7 +469,7 @@ void FVulkanPipelineStateCacheManager::InitAndLoad(const TArray<FString>& CacheF
 	{
 		VkPipelineCacheCreateInfo PipelineCacheInfo;
 		ZeroVulkanStruct(PipelineCacheInfo, VK_STRUCTURE_TYPE_PIPELINE_CACHE_CREATE_INFO);
-		VERIFYVULKANRESULT(VulkanRHI::vkCreatePipelineCache(Device->GetInstanceHandle(), &PipelineCacheInfo, VULKAN_CPU_ALLOCATOR, &PipelineCacheExclusive.Get()));
+		VERIFYVULKANRESULT(VulkanRHI::vkCreatePipelineCache(Device->GetHandle(), &PipelineCacheInfo, VULKAN_CPU_ALLOCATOR, &PipelineCacheExclusive.Get()));
 	}
 }
 
@@ -521,7 +521,7 @@ void FVulkanPipelineStateCacheManager::OnShaderPipelineCacheOpened(FString const
 
 		VkPipelineCacheCreateInfo PipelineCacheInfo;
 		ZeroVulkanStruct(PipelineCacheInfo, VK_STRUCTURE_TYPE_PIPELINE_CACHE_CREATE_INFO);
-		VERIFYVULKANRESULT(VulkanRHI::vkCreatePipelineCache(Device->GetInstanceHandle(), &PipelineCacheInfo, VULKAN_CPU_ALLOCATOR, &PipelineCacheExclusive.Get()));
+		VERIFYVULKANRESULT(VulkanRHI::vkCreatePipelineCache(Device->GetHandle(), &PipelineCacheInfo, VULKAN_CPU_ALLOCATOR, &PipelineCacheExclusive.Get()));
 	}
 
 	if (PlatformFile.FileExists(*CompiledPSOCacheFolderName))
@@ -595,8 +595,8 @@ void FVulkanPipelineStateCacheManager::OnShaderPipelineCachePrecompilationComple
 		QUICK_SCOPE_CYCLE_COUNTER(STAT_VulkanPSOCacheMerge);
 		FScopedPipelineCache GlobalPipelineCacheExclusive = GlobalPSOCache.Get(EPipelineCacheAccess::Exclusive);
 		FScopedPipelineCache CurrentPipelineCacheExclusive = CurrentPrecompilingPSOCache.Get(EPipelineCacheAccess::Exclusive);
-		VERIFYVULKANRESULT(VulkanRHI::vkMergePipelineCaches(Device->GetInstanceHandle(), GlobalPipelineCacheExclusive.Get(), 1, &CurrentPipelineCacheExclusive.Get()));
-		VulkanRHI::vkDestroyPipelineCache(Device->GetInstanceHandle(), CurrentPipelineCacheExclusive.Get(), VULKAN_CPU_ALLOCATOR);
+		VERIFYVULKANRESULT(VulkanRHI::vkMergePipelineCaches(Device->GetHandle(), GlobalPipelineCacheExclusive.Get(), 1, &CurrentPipelineCacheExclusive.Get()));
+		VulkanRHI::vkDestroyPipelineCache(Device->GetHandle(), CurrentPipelineCacheExclusive.Get(), VULKAN_CPU_ALLOCATOR);
 
 		CurrentPipelineCacheExclusive.Get() = VK_NULL_HANDLE;
 	}
@@ -614,13 +614,13 @@ void FVulkanPipelineStateCacheManager::SavePSOCache(const FString& CacheFilename
 
 	// First save Device Cache
 	size_t Size = 0;
-	VERIFYVULKANRESULT(VulkanRHI::vkGetPipelineCacheData(Device->GetInstanceHandle(), PipelineCacheExclusive.Get(), &Size, nullptr));
+	VERIFYVULKANRESULT(VulkanRHI::vkGetPipelineCacheData(Device->GetHandle(), PipelineCacheExclusive.Get(), &Size, nullptr));
 	// 16 is HeaderSize + HeaderVersion
 	if (Size >= 16 + VK_UUID_SIZE)
 	{
 		TArray<uint8> DeviceCache;
 		DeviceCache.AddUninitialized(Size);
-		VkResult Result = VulkanRHI::vkGetPipelineCacheData(Device->GetInstanceHandle(), PipelineCacheExclusive.Get(), &Size, DeviceCache.GetData());
+		VkResult Result = VulkanRHI::vkGetPipelineCacheData(Device->GetHandle(), PipelineCacheExclusive.Get(), &Size, DeviceCache.GetData());
 		if (Result == VK_SUCCESS)
 		{
 			FString BinaryCacheFilename = FVulkanPlatform::CreatePSOBinaryCacheFilename(Device, CacheFilename);
@@ -639,10 +639,10 @@ void FVulkanPipelineStateCacheManager::SavePSOCache(const FString& CacheFilename
 			UE_LOG(LogVulkanRHI, Warning, TEXT("Failed to get Vulkan pipeline cache data. Error %d, %d bytes"), Result, Size);
 
 			//TODO: Resave it when we shutdown the manager?!
-			VulkanRHI::vkDestroyPipelineCache(Device->GetInstanceHandle(), PipelineCacheExclusive.Get(), VULKAN_CPU_ALLOCATOR);
+			VulkanRHI::vkDestroyPipelineCache(Device->GetHandle(), PipelineCacheExclusive.Get(), VULKAN_CPU_ALLOCATOR);
 			VkPipelineCacheCreateInfo PipelineCacheInfo;
 			ZeroVulkanStruct(PipelineCacheInfo, VK_STRUCTURE_TYPE_PIPELINE_CACHE_CREATE_INFO);
-			VERIFYVULKANRESULT(VulkanRHI::vkCreatePipelineCache(Device->GetInstanceHandle(), &PipelineCacheInfo, VULKAN_CPU_ALLOCATOR, &PipelineCacheExclusive.Get()));
+			VERIFYVULKANRESULT(VulkanRHI::vkCreatePipelineCache(Device->GetHandle(), &PipelineCacheInfo, VULKAN_CPU_ALLOCATOR, &PipelineCacheExclusive.Get()));
 		}
 		else
 		{
@@ -1200,7 +1200,7 @@ static FString GfxShaderHashesToString(FVulkanShader* Shaders[ShaderStage::NumGr
 	return ShaderHashes;
 }
 
-bool FVulkanPipelineStateCacheManager::CreateGfxPipelineFromEntry(FVulkanRHIGraphicsPipelineState* PSO, FVulkanShader* Shaders[ShaderStage::NumGraphicsStages], FGraphicsPipelineStateInitializer::EPSOPrecacheCompileType PSOCompileType)
+bool FVulkanPipelineStateCacheManager::CreateGfxPipelineFromEntry(FVulkanGraphicsPipelineState* PSO, FVulkanShader* Shaders[ShaderStage::NumGraphicsStages], FGraphicsPipelineStateInitializer::EPSOPrecacheCompileType PSOCompileType)
 {
 	VkPipeline* Pipeline = &PSO->VulkanPipeline;
 	const FGfxPipelineDesc* GfxEntry = &PSO->Desc;
@@ -1445,7 +1445,7 @@ bool FVulkanPipelineStateCacheManager::CreateGfxPipelineFromEntry(FVulkanRHIGrap
 
 #if PLATFORM_ANDROID
 
-static VkResult CreatePSOWithExternalService(FVulkanDevice* Device, FGraphicsPipelineStateInitializer::EPSOPrecacheCompileType PSOCompileType, FVulkanRHIGraphicsPipelineState* PSO, FVulkanShader* Shaders[ShaderStage::NumGraphicsStages], const VkGraphicsPipelineCreateInfo& PipelineInfo, VkPipelineCache DestPipelineCache, FRWLock& PipelineLock)
+static VkResult CreatePSOWithExternalService(FVulkanDevice* Device, FGraphicsPipelineStateInitializer::EPSOPrecacheCompileType PSOCompileType, FVulkanGraphicsPipelineState* PSO, FVulkanShader* Shaders[ShaderStage::NumGraphicsStages], const VkGraphicsPipelineCreateInfo& PipelineInfo, VkPipelineCache DestPipelineCache, FRWLock& PipelineLock)
 {
 	VkResult Result = VK_ERROR_INITIALIZATION_FAILED;
 	FVulkanShader::FSpirvCode VS = PSO->GetPatchedSpirvCode(Shaders[ShaderStage::Vertex]);
@@ -1464,9 +1464,9 @@ static VkResult CreatePSOWithExternalService(FVulkanDevice* Device, FGraphicsPip
 		QUICK_SCOPE_CYCLE_COUNTER(STAT_Vulkan_RHICreateGraphicsPipelineState_ExternalInitialCache);
 		FRWScopeLock Lock(PipelineLock, SLT_Write);
 		size_t InitialCacheSize = 0;
-		VulkanRHI::vkGetPipelineCacheData(Device->GetInstanceHandle(), DestPipelineCache, &InitialCacheSize, nullptr);
+		VulkanRHI::vkGetPipelineCacheData(Device->GetHandle(), DestPipelineCache, &InitialCacheSize, nullptr);
 		InitialCacheData.SetNumUninitialized(InitialCacheSize);
-		VulkanRHI::vkGetPipelineCacheData(Device->GetInstanceHandle(), DestPipelineCache, &InitialCacheSize, InitialCacheData.GetData());
+		VulkanRHI::vkGetPipelineCacheData(Device->GetHandle(), DestPipelineCache, &InitialCacheSize, InitialCacheData.GetData());
 	}
 	FString FailLog;
 	LocalPipelineCache = FVulkanPlatform::PrecompilePSO(Device, InitialCacheData, PSOCompileType, &PipelineInfo, GfxEntry, &PSO->RenderPass->GetLayout(), VSCode, PSCode, AfterSize,&FailLog);
@@ -1477,9 +1477,9 @@ static VkResult CreatePSOWithExternalService(FVulkanDevice* Device, FGraphicsPip
 			QUICK_SCOPE_CYCLE_COUNTER(STAT_Vulkan_RHICreateGraphicsPipelineState_ExternalMergeResult);
 			FRWScopeLock Lock(PipelineLock, SLT_Write);
 			Result = VK_SUCCESS;
-			VERIFYVULKANRESULT(VulkanRHI::vkMergePipelineCaches(Device->GetInstanceHandle(), DestPipelineCache, 1, &LocalPipelineCache));
+			VERIFYVULKANRESULT(VulkanRHI::vkMergePipelineCaches(Device->GetHandle(), DestPipelineCache, 1, &LocalPipelineCache));
 		}
-		VulkanRHI::vkDestroyPipelineCache(Device->GetInstanceHandle(), LocalPipelineCache, VULKAN_CPU_ALLOCATOR);
+		VulkanRHI::vkDestroyPipelineCache(Device->GetHandle(), LocalPipelineCache, VULKAN_CPU_ALLOCATOR);
 	}
 	else
 	{
@@ -1489,7 +1489,7 @@ static VkResult CreatePSOWithExternalService(FVulkanDevice* Device, FGraphicsPip
 }
 #endif
 
-VkResult FVulkanPipelineStateCacheManager::CreateVKPipeline(FVulkanRHIGraphicsPipelineState* PSO, FVulkanShader* Shaders[ShaderStage::NumGraphicsStages], const VkGraphicsPipelineCreateInfo& PipelineInfo, FGraphicsPipelineStateInitializer::EPSOPrecacheCompileType PSOCompileType)
+VkResult FVulkanPipelineStateCacheManager::CreateVKPipeline(FVulkanGraphicsPipelineState* PSO, FVulkanShader* Shaders[ShaderStage::NumGraphicsStages], const VkGraphicsPipelineCreateInfo& PipelineInfo, FGraphicsPipelineStateInitializer::EPSOPrecacheCompileType PSOCompileType)
 {
 	bool bIsPrecompileJob = PSOCompileType != FGraphicsPipelineStateInitializer::EPSOPrecacheCompileType::NotSet;
 	if(FVulkanChunkedPipelineCacheManager::IsEnabled())
@@ -1499,10 +1499,10 @@ VkResult FVulkanPipelineStateCacheManager::CreateVKPipeline(FVulkanRHIGraphicsPi
 		// Use chunk caching and bypass FVulkanPipelineStateCacheManager's PSO caching
 		// Placeholder PSO size - TODO: remove pipeline cache size stuff.	
 		PSO->PipelineCacheSize = 20 * 1024; // This is only required bUseLRU == true.
-		return FVulkanChunkedPipelineCacheManager::Get().CreatePSO(PSO, bIsPrecompileJob, FVulkanChunkedPipelineCacheManager::FPSOCreateCallbackFunc<FVulkanRHIGraphicsPipelineState>(
-				[&](FVulkanChunkedPipelineCacheManager::FPSOCreateFuncParams<FVulkanRHIGraphicsPipelineState>& Params)
+		return FVulkanChunkedPipelineCacheManager::Get().CreatePSO(PSO, bIsPrecompileJob, FVulkanChunkedPipelineCacheManager::FPSOCreateCallbackFunc<FVulkanGraphicsPipelineState>(
+				[&](FVulkanChunkedPipelineCacheManager::FPSOCreateFuncParams<FVulkanGraphicsPipelineState>& Params)
 			{
-				FVulkanRHIGraphicsPipelineState* PSO = Params.PSO;
+				FVulkanGraphicsPipelineState* PSO = Params.PSO;
 				VkPipelineCache& PipelineCache = Params.DestPipelineCache;
 				FVulkanChunkedPipelineCacheManager::EPSOOperation PSOOperation = Params.PSOOperation;
 				
@@ -1519,7 +1519,7 @@ VkResult FVulkanPipelineStateCacheManager::CreateVKPipeline(FVulkanRHIGraphicsPi
 						FRWScopeLock Lock(Params.DestPipelineCacheLock, SLT_ReadOnly);
 						VkGraphicsPipelineCreateInfo TestPipelineInfo = PipelineInfo;
 						TestPipelineInfo.flags |= VK_PIPELINE_CREATE_FAIL_ON_PIPELINE_COMPILE_REQUIRED_BIT_EXT;
-						Result = VulkanRHI::vkCreateGraphicsPipelines(Device->GetInstanceHandle(), PipelineCache, 1, &TestPipelineInfo, VULKAN_CPU_ALLOCATOR, &PSO->VulkanPipeline);
+						Result = VulkanRHI::vkCreateGraphicsPipelines(Device->GetHandle(), PipelineCache, 1, &TestPipelineInfo, VULKAN_CPU_ALLOCATOR, &PSO->VulkanPipeline);
 					}
 					else
 					{
@@ -1542,7 +1542,7 @@ VkResult FVulkanPipelineStateCacheManager::CreateVKPipeline(FVulkanRHIGraphicsPi
 				{
 					QUICK_SCOPE_CYCLE_COUNTER(STAT_Vulkan_RHICreateGraphicsPipelineState_vkCreate);
 					FRWScopeLock Lock(Params.DestPipelineCacheLock, SLT_ReadOnly);
-					Result = VulkanRHI::vkCreateGraphicsPipelines(Device->GetInstanceHandle(), PipelineCache, 1, &PipelineInfo, VULKAN_CPU_ALLOCATOR, &PSO->VulkanPipeline);
+					Result = VulkanRHI::vkCreateGraphicsPipelines(Device->GetHandle(), PipelineCache, 1, &PipelineInfo, VULKAN_CPU_ALLOCATOR, &PSO->VulkanPipeline);
 				}
 				return Result;
 			}));
@@ -1605,7 +1605,7 @@ VkResult FVulkanPipelineStateCacheManager::CreateVKPipeline(FVulkanRHIGraphicsPi
 			{
 				Result = VK_ERROR_INITIALIZATION_FAILED;
 				bWantPSOSize = true;
-				VulkanRHI::vkDestroyPipelineCache(Device->GetInstanceHandle(), LocalPipelineCache, VULKAN_CPU_ALLOCATOR);
+				VulkanRHI::vkDestroyPipelineCache(Device->GetHandle(), LocalPipelineCache, VULKAN_CPU_ALLOCATOR);
 				LocalPipelineCache = VK_NULL_HANDLE;
 			}
 
@@ -1633,14 +1633,14 @@ VkResult FVulkanPipelineStateCacheManager::CreateVKPipeline(FVulkanRHIGraphicsPi
 			// measuring deltas from the global PipelineCache is not thread safe.
 			VkPipelineCacheCreateInfo PipelineCacheInfo;
 			ZeroVulkanStruct(PipelineCacheInfo, VK_STRUCTURE_TYPE_PIPELINE_CACHE_CREATE_INFO);
-			VERIFYVULKANRESULT(VulkanRHI::vkCreatePipelineCache(Device->GetInstanceHandle(), &PipelineCacheInfo, VULKAN_CPU_ALLOCATOR, &LocalPipelineCache));
-			Result = VulkanRHI::vkCreateGraphicsPipelines(Device->GetInstanceHandle(), LocalPipelineCache, 1, &PipelineInfo, VULKAN_CPU_ALLOCATOR, Pipeline);
+			VERIFYVULKANRESULT(VulkanRHI::vkCreatePipelineCache(Device->GetHandle(), &PipelineCacheInfo, VULKAN_CPU_ALLOCATOR, &LocalPipelineCache));
+			Result = VulkanRHI::vkCreateGraphicsPipelines(Device->GetHandle(), LocalPipelineCache, 1, &PipelineInfo, VULKAN_CPU_ALLOCATOR, Pipeline);
 			if (bValidateServicePSO)
 			{
 				size_t Diff = 0;
 				if (ensure(LocalPipelineCache != VK_NULL_HANDLE))
 				{
-					VulkanRHI::vkGetPipelineCacheData(Device->GetInstanceHandle(), LocalPipelineCache, &Diff, nullptr);
+					VulkanRHI::vkGetPipelineCacheData(Device->GetHandle(), LocalPipelineCache, &Diff, nullptr);
 				}
 				UE_CLOG(Diff != PSOSize, LogVulkanRHI, Warning, TEXT("PSO service size mismatches engine size! [PSOService = %d, Game Process = %d]"), PSOSize, Diff);
 			}
@@ -1649,7 +1649,7 @@ VkResult FVulkanPipelineStateCacheManager::CreateVKPipeline(FVulkanRHIGraphicsPi
 		{
 			SCOPE_CYCLE_COUNTER(STAT_VulkanPSOVulkanCreationTime);
 			FScopedPipelineCache PipelineCacheShared = Cache.Get(EPipelineCacheAccess::Shared);
-			Result = VulkanRHI::vkCreateGraphicsPipelines(Device->GetInstanceHandle(), PipelineCacheShared.Get(), 1, &PipelineInfo, VULKAN_CPU_ALLOCATOR, Pipeline);
+			Result = VulkanRHI::vkCreateGraphicsPipelines(Device->GetHandle(), PipelineCacheShared.Get(), 1, &PipelineInfo, VULKAN_CPU_ALLOCATOR, Pipeline);
 		}
 	}
 
@@ -1666,7 +1666,7 @@ VkResult FVulkanPipelineStateCacheManager::CreateVKPipeline(FVulkanRHIGraphicsPi
 
 				if (LocalPipelineCache != VK_NULL_HANDLE)
 				{
-					VulkanRHI::vkGetPipelineCacheData(Device->GetInstanceHandle(), LocalPipelineCache, &Diff, nullptr);
+					VulkanRHI::vkGetPipelineCacheData(Device->GetHandle(), LocalPipelineCache, &Diff, nullptr);
 				}
 
 				if (!Diff)
@@ -1690,9 +1690,9 @@ VkResult FVulkanPipelineStateCacheManager::CreateVKPipeline(FVulkanRHIGraphicsPi
 		if (PipelineCacheExclusive.Get() != VK_NULL_HANDLE)
 		{
 			QUICK_SCOPE_CYCLE_COUNTER(STAT_VulkanPSOCacheMerge);
-			VERIFYVULKANRESULT(VulkanRHI::vkMergePipelineCaches(Device->GetInstanceHandle(), PipelineCacheExclusive.Get(), 1, &LocalPipelineCache));
+			VERIFYVULKANRESULT(VulkanRHI::vkMergePipelineCaches(Device->GetHandle(), PipelineCacheExclusive.Get(), 1, &LocalPipelineCache));
 		}
-		VulkanRHI::vkDestroyPipelineCache(Device->GetInstanceHandle(), LocalPipelineCache, VULKAN_CPU_ALLOCATOR);
+		VulkanRHI::vkDestroyPipelineCache(Device->GetHandle(), LocalPipelineCache, VULKAN_CPU_ALLOCATOR);
 	}
 
 	VERIFYVULKANRESULT(Result);
@@ -1710,7 +1710,7 @@ VkResult FVulkanPipelineStateCacheManager::CreateRayTracingPipeline(const VkRayT
 		if (PipelineCacheShared.Get() != VK_NULL_HANDLE)
 		{
 			return VulkanDynamicAPI::vkCreateRayTracingPipelinesKHR(
-				Device->GetInstanceHandle(),
+				Device->GetHandle(),
 				VK_NULL_HANDLE, // Deferred Operation
 				PipelineCacheShared.Get(), // Pipeline Cache 
 				1,
@@ -1722,7 +1722,7 @@ VkResult FVulkanPipelineStateCacheManager::CreateRayTracingPipeline(const VkRayT
 
 	FScopedPipelineCache PipelineCacheShared = GlobalPSOCache.Get(EPipelineCacheAccess::Shared);
 	return VulkanDynamicAPI::vkCreateRayTracingPipelinesKHR(
-		Device->GetInstanceHandle(),
+		Device->GetHandle(),
 		VK_NULL_HANDLE, // Deferred Operation
 		PipelineCacheShared.Get(), // Pipeline Cache 
 		1,
@@ -1733,14 +1733,14 @@ VkResult FVulkanPipelineStateCacheManager::CreateRayTracingPipeline(const VkRayT
 
 void FVulkanPipelineStateCacheManager::DestroyCache()
 {
-	VkDevice DeviceHandle = Device->GetInstanceHandle();
+	VkDevice DeviceHandle = Device->GetHandle();
 
 	QUICK_SCOPE_CYCLE_COUNTER(STAT_Vulkan_DestroyCache_PSOLock);
 	FScopeLock Lock1(&GraphicsPSOLockedCS);
 	int idx = 0;
 	for (auto& Pair : GraphicsPSOLockedMap)
 	{
-		FVulkanRHIGraphicsPipelineState* Pipeline = Pair.Value;
+		FVulkanGraphicsPipelineState* Pipeline = Pair.Value;
 		FPlatformMisc::LowLevelOutputDebugStringf(TEXT("Leaked PSO %05d: RefCount=%d Handle=0x%p\n"), idx++, Pipeline->GetRefCount(), Pipeline);
 	}
 	LRU2SizeList.Reset();
@@ -2011,7 +2011,7 @@ void FVulkanPipelineStateCacheManager::CreateGfxEntry(const FGraphicsPipelineSta
 
 
 
-FVulkanRHIGraphicsPipelineState::FVulkanRHIGraphicsPipelineState(FVulkanDevice* InDevice, const FGraphicsPipelineStateInitializer& PSOInitializer_, const FGfxPipelineDesc& InDesc, FVulkanPSOKey* VulkanKey)
+FVulkanGraphicsPipelineState::FVulkanGraphicsPipelineState(FVulkanDevice* InDevice, const FGraphicsPipelineStateInitializer& PSOInitializer_, const FGfxPipelineDesc& InDesc, FVulkanPSOKey* VulkanKey)
 	: bIsRegistered(false)
 	, PrimitiveType(PSOInitializer_.PrimitiveType)
 	, VulkanPipeline(0)
@@ -2070,14 +2070,14 @@ FVulkanRHIGraphicsPipelineState::FVulkanRHIGraphicsPipelineState(FVulkanDevice*
 
 void FVulkanPipelineStateCacheManager::NotifyDeletedGraphicsPSO(FRHIGraphicsPipelineState* PSO)
 {
-	FVulkanRHIGraphicsPipelineState* VkPSO = (FVulkanRHIGraphicsPipelineState*)PSO;
+	FVulkanGraphicsPipelineState* VkPSO = (FVulkanGraphicsPipelineState*)PSO;
 	Device->NotifyDeletedGfxPipeline(VkPSO);
 	FVulkanPSOKey& Key = VkPSO->VulkanKey;
 	DEC_DWORD_STAT_BY(STAT_VulkanPSOKeyMemory, Key.GetDataRef().Num());
 	if(VkPSO->bIsRegistered)
 	{
 		FScopeLock Lock(&GraphicsPSOLockedCS);
-		FVulkanRHIGraphicsPipelineState** Contained = GraphicsPSOLockedMap.Find(Key);
+		FVulkanGraphicsPipelineState** Contained = GraphicsPSOLockedMap.Find(Key);
 		check(Contained && *Contained == PSO);
 		VkPSO->bIsRegistered = false;
 		if(bUseLRU)
@@ -2095,7 +2095,7 @@ void FVulkanPipelineStateCacheManager::NotifyDeletedGraphicsPSO(FRHIGraphicsPipe
 	else
 	{
 		FScopeLock Lock(&GraphicsPSOLockedCS);
-		FVulkanRHIGraphicsPipelineState** Contained = GraphicsPSOLockedMap.Find(Key);
+		FVulkanGraphicsPipelineState** Contained = GraphicsPSOLockedMap.Find(Key);
 		if (Contained && *Contained == VkPSO)
 		{
 			check(0);
@@ -2134,12 +2134,12 @@ FGraphicsPipelineStateRHIRef FVulkanPipelineStateCacheManager::RHICreateGraphics
 	}
 
 
-	FVulkanRHIGraphicsPipelineState* NewPSO = 0;
+	FVulkanGraphicsPipelineState* NewPSO = 0;
 	{
 		SCOPE_CYCLE_COUNTER(STAT_VulkanPSOLookupTime);
 		FScopeLock Lock(&GraphicsPSOLockedCS);
 		{
-			FVulkanRHIGraphicsPipelineState** PSO = GraphicsPSOLockedMap.Find(Key);
+			FVulkanGraphicsPipelineState** PSO = GraphicsPSOLockedMap.Find(Key);
 			if(PSO)
 			{
 				check(*PSO);
@@ -2157,7 +2157,7 @@ FGraphicsPipelineStateRHIRef FVulkanPipelineStateCacheManager::RHICreateGraphics
 	{
 		// Workers can be creating PSOs while FRHIResource::FlushPendingDeletes is running on the RHI thread
 		// so let it get enqueued for a delete with Release() instead.  Only used for failed or duplicate PSOs...
-		auto DeleteNewPSO = [](FVulkanRHIGraphicsPipelineState* PSOPtr)
+		auto DeleteNewPSO = [](FVulkanGraphicsPipelineState* PSOPtr)
 		{
 			PSOPtr->AddRef();
 			const uint32 RefCount = PSOPtr->Release();
@@ -2165,7 +2165,7 @@ FGraphicsPipelineStateRHIRef FVulkanPipelineStateCacheManager::RHICreateGraphics
 		};
 
 		SCOPE_CYCLE_COUNTER(STAT_VulkanPSOCreationTime);
-		NewPSO = new FVulkanRHIGraphicsPipelineState(Device, Initializer, Desc, &Key);
+		NewPSO = new FVulkanGraphicsPipelineState(Device, Initializer, Desc, &Key);
 		{
 			FVulkanLayout* Layout = FindOrAddLayout(DescriptorSetLayoutInfo, true, NewPSO->UsesBindless());
 			NewPSO->Layout = Layout;
@@ -2227,7 +2227,7 @@ FGraphicsPipelineStateRHIRef FVulkanPipelineStateCacheManager::RHICreateGraphics
 				}
 			}
 			FScopeLock Lock(&GraphicsPSOLockedCS); 
-			FVulkanRHIGraphicsPipelineState** MapPSO = GraphicsPSOLockedMap.Find(Key);
+			FVulkanGraphicsPipelineState** MapPSO = GraphicsPSOLockedMap.Find(Key);
 			if(MapPSO)//another thread could end up creating it.
 			{
 				DeleteNewPSO(NewPSO);
@@ -2378,7 +2378,7 @@ FVulkanComputePipeline* FVulkanPipelineStateCacheManager::CreateComputePipelineF
 	{
 		QUICK_SCOPE_CYCLE_COUNTER(STAT_VulkanComputePSOCreate);
 		FScopedPipelineCache PipelineCacheShared = GlobalPSOCache.Get(EPipelineCacheAccess::Shared);
-		Result = VulkanRHI::vkCreateComputePipelines(Device->GetInstanceHandle(), PipelineCacheShared.Get(), 1, &PipelineInfo, VULKAN_CPU_ALLOCATOR, &Pipeline->Pipeline);
+		Result = VulkanRHI::vkCreateComputePipelines(Device->GetHandle(), PipelineCacheShared.Get(), 1, &PipelineInfo, VULKAN_CPU_ALLOCATOR, &Pipeline->Pipeline);
 	}
 
 	if (Result != VK_SUCCESS)
@@ -2505,7 +2505,7 @@ void GetVulkanGfxShaders(const FBoundShaderStateInput& BSI, FVulkanShader* OutSh
 	}
 }
 
-void GetVulkanGfxShaders(FVulkanDevice* Device, const FVulkanRHIGraphicsPipelineState& GfxPipelineState, FVulkanShader* OutShaders[ShaderStage::NumGraphicsStages])
+void GetVulkanGfxShaders(FVulkanDevice* Device, const FVulkanGraphicsPipelineState& GfxPipelineState, FVulkanShader* OutShaders[ShaderStage::NumGraphicsStages])
 {
 	FMemory::Memzero(OutShaders, ShaderStage::NumGraphicsStages * sizeof(*OutShaders));
 	Device->GetShaderFactory().LookupGfxShaders(GfxPipelineState.ShaderKeys, OutShaders);
@@ -2527,13 +2527,13 @@ void FVulkanPipelineStateCacheManager::TickLRU()
 	const int MaxEvictsPerTick = 5;
 	for(int i = 0 ; i<MaxEvictsPerTick; i++)
 	{
-		FVulkanRHIGraphicsPipelineStateLRUNode* Node = LRU.GetTail();
+		FVulkanGraphicsPipelineStateLRUNode* Node = LRU.GetTail();
 		if (!Node)
 		{
 			return;
 		}
 
-		TRefCountPtr<FVulkanRHIGraphicsPipelineState> PSO = Node->GetValue();
+		TRefCountPtr<FVulkanGraphicsPipelineState> PSO = Node->GetValue();
 		bool bTimeToDie = PSO->LRUFrame + GVulkanPSOLRUEvictAfterUnusedFrames < GFrameNumberRenderThread;
 		if (bTimeToDie)
 		{
@@ -2553,12 +2553,12 @@ void FVulkanPipelineStateCacheManager::LRUDump()
 #if !UE_BUILD_SHIPPING
 	uint32 tid = FPlatformTLS::GetCurrentThreadId();
 	LRUPRINT(TEXT("//***** LRU DUMP *****\\\\\n"));
-	FVulkanRHIGraphicsPipelineStateLRUNode* Node= LRU.GetHead();
+	FVulkanGraphicsPipelineStateLRUNode* Node= LRU.GetHead();
 	uint32_t Size = 0;
 	uint32_t Index = 0;
 	while(Node)
 	{
-		FVulkanRHIGraphicsPipelineState* PSO = Node->GetValue();
+		FVulkanGraphicsPipelineState* PSO = Node->GetValue();
 		Size += PSO->PipelineCacheSize;
 		LRUPRINT(TEXT("\t%08x PSO %p :: %d  :: %06d \\ %06d\n"), tid, PSO, PSO->LRUFrame, PSO->PipelineCacheSize, Size);
 		Node = Node->GetNextNode();
@@ -2602,7 +2602,7 @@ void FVulkanPipelineStateCacheManager::LRUDebugEvictAll()
 	LRUPRINT_DEBUG(TEXT("Evicted %d\n"), Count);
 }
 
-void FVulkanPipelineStateCacheManager::LRUAdd(FVulkanRHIGraphicsPipelineState* PSO)
+void FVulkanPipelineStateCacheManager::LRUAdd(FVulkanGraphicsPipelineState* PSO)
 {
 	if(!bUseLRU)
 	{
@@ -2629,7 +2629,7 @@ void FVulkanPipelineStateCacheManager::LRUAdd(FVulkanRHIGraphicsPipelineState* P
 
 }
 
-void FVulkanPipelineStateCacheManager::LRUTouch(FVulkanRHIGraphicsPipelineState* PSO)
+void FVulkanPipelineStateCacheManager::LRUTouch(FVulkanGraphicsPipelineState* PSO)
 {
 	if(!bUseLRU)
 	{
@@ -2685,13 +2685,13 @@ void FVulkanPipelineStateCacheManager::LRUTouch(FVulkanRHIGraphicsPipelineState*
 	}
 }
 
-void FVulkanRHIGraphicsPipelineState::DeleteVkPipeline(bool bImmediate)
+void FVulkanGraphicsPipelineState::DeleteVkPipeline(bool bImmediate)
 {
 	if (VulkanPipeline != VK_NULL_HANDLE)
 	{
 		if (bImmediate)
 		{
-			VulkanRHI::vkDestroyPipeline(Device->GetInstanceHandle(), VulkanPipeline, VULKAN_CPU_ALLOCATOR);
+			VulkanRHI::vkDestroyPipeline(Device->GetHandle(), VulkanPipeline, VULKAN_CPU_ALLOCATOR);
 		}
 		else
 		{
@@ -2703,17 +2703,17 @@ void FVulkanRHIGraphicsPipelineState::DeleteVkPipeline(bool bImmediate)
 	Device->PipelineStateCache->LRUCheckNotInside(this);
 }
 
-void FVulkanPipelineStateCacheManager::LRUCheckNotInside(FVulkanRHIGraphicsPipelineState* PSO)
+void FVulkanPipelineStateCacheManager::LRUCheckNotInside(FVulkanGraphicsPipelineState* PSO)
 {
 	FScopeLock Lock(&LRUCS);
 
 
-	FVulkanRHIGraphicsPipelineStateLRUNode* Node = LRU.GetHead();
+	FVulkanGraphicsPipelineStateLRUNode* Node = LRU.GetHead();
 	uint32_t Size = 0;
 	uint32_t Index = 0;
 	while (Node)
 	{
-		FVulkanRHIGraphicsPipelineState* foo = Node->GetValue();
+		FVulkanGraphicsPipelineState* foo = Node->GetValue();
 		if (foo == PSO)
 		{
 			check(0 == foo->LRUNode);
@@ -2725,7 +2725,7 @@ void FVulkanPipelineStateCacheManager::LRUCheckNotInside(FVulkanRHIGraphicsPipel
 	check(0 == PSO->LRUNode);
 }
 
-void FVulkanPipelineStateCacheManager::LRURemove(FVulkanRHIGraphicsPipelineState* PSO)
+void FVulkanPipelineStateCacheManager::LRURemove(FVulkanGraphicsPipelineState* PSO)
 {
 	check(bUseLRU);
 	if (PSO->LRUNode != 0)
@@ -2761,9 +2761,9 @@ bool FVulkanPipelineStateCacheManager::LRUEvictOne(bool bOnlyOld)
 {
 	check(bUseLRU);
 	uint32 tid = FPlatformTLS::GetCurrentThreadId();
-	FVulkanRHIGraphicsPipelineStateLRUNode* Node = LRU.GetTail();
+	FVulkanGraphicsPipelineStateLRUNode* Node = LRU.GetTail();
 	check(Node != 0);
-	TRefCountPtr<FVulkanRHIGraphicsPipelineState> PSO = Node->GetValue();
+	TRefCountPtr<FVulkanGraphicsPipelineState> PSO = Node->GetValue();
 
 	bool bImmediate = PSO->LRUFrame + 3 < GFrameNumberRenderThread;
 	if(bOnlyOld && !bImmediate)
diff --git a/Engine/Source/Runtime/VulkanRHI/Private/VulkanPipeline.h b/Engine/Source/Runtime/VulkanRHI/Private/VulkanPipeline.h
--- a/Engine/Source/Runtime/VulkanRHI/Private/VulkanPipeline.h
+++ b/Engine/Source/Runtime/VulkanRHI/Private/VulkanPipeline.h
@@ -15,11 +15,11 @@
 #define VULKAN_PSO_CACHE_DEBUG 0
 #define VULKAN_USE_SHADERKEYS 1
 class FVulkanDevice;
-class FVulkanRHIGraphicsPipelineState;
+class FVulkanGraphicsPipelineState;
 
 
-using FVulkanRHIGraphicsPipelineStateLRU = TDoubleLinkedList<FVulkanRHIGraphicsPipelineState*>;
-using FVulkanRHIGraphicsPipelineStateLRUNode = FVulkanRHIGraphicsPipelineStateLRU::TDoubleLinkedListNode;
+using FVulkanGraphicsPipelineStateLRU = TDoubleLinkedList<FVulkanGraphicsPipelineState*>;
+using FVulkanGraphicsPipelineStateLRUNode = FVulkanGraphicsPipelineStateLRU::TDoubleLinkedListNode;
 
 
 class FVulkanPSOKey : public TDataKey<FVulkanPSOKey, true> {};
@@ -55,7 +55,7 @@ inline uint64 GetShaderKeyForGfxStage(const FBoundShaderStateInput& BSI, ShaderS
 
 
 void GetVulkanGfxShaders(const FBoundShaderStateInput& BSI, FVulkanShader* OutShaders[ShaderStage::NumGraphicsStages]);
-void GetVulkanGfxShaders(FVulkanDevice* Device, const FVulkanRHIGraphicsPipelineState& GfxPipelineState, FVulkanShader* OutShaders[ShaderStage::NumGraphicsStages]);
+void GetVulkanGfxShaders(FVulkanDevice* Device, const FVulkanGraphicsPipelineState& GfxPipelineState, FVulkanShader* OutShaders[ShaderStage::NumGraphicsStages]);
 
 
 
@@ -520,9 +520,9 @@ private:
 	FGraphicsPipelineStateRHIRef RHICreateGraphicsPipelineState(const FGraphicsPipelineStateInitializer& Initializer);
 	FVulkanComputePipeline* RHICreateComputePipelineState(FRHIComputeShader* ComputeShaderRHI);
 	void NotifyDeletedGraphicsPSO(FRHIGraphicsPipelineState* PSO);
-	bool CreateGfxPipelineFromEntry(FVulkanRHIGraphicsPipelineState* PSO, FVulkanShader* Shaders[ShaderStage::NumGraphicsStages], FGraphicsPipelineStateInitializer::EPSOPrecacheCompileType PSOCompileType);
+	bool CreateGfxPipelineFromEntry(FVulkanGraphicsPipelineState* PSO, FVulkanShader* Shaders[ShaderStage::NumGraphicsStages], FGraphicsPipelineStateInitializer::EPSOPrecacheCompileType PSOCompileType);
 
-	VkResult CreateVKPipeline(FVulkanRHIGraphicsPipelineState* PSO, FVulkanShader* Shaders[ShaderStage::NumGraphicsStages], const VkGraphicsPipelineCreateInfo& PipelineInfo, FGraphicsPipelineStateInitializer::EPSOPrecacheCompileType PSOCompileType);
+	VkResult CreateVKPipeline(FVulkanGraphicsPipelineState* PSO, FVulkanShader* Shaders[ShaderStage::NumGraphicsStages], const VkGraphicsPipelineCreateInfo& PipelineInfo, FGraphicsPipelineStateInitializer::EPSOPrecacheCompileType PSOCompileType);
 
 	FVulkanLayout* FindOrAddLayout(const FVulkanDescriptorSetsLayoutInfo& DescriptorSetLayoutInfo, bool bGfxLayout, bool bUsesBindless);
 	FVulkanComputePipeline* CreateComputePipelineFromShader(FVulkanComputeShader* Shader);
@@ -531,14 +531,14 @@ private:
 	void TickLRU();
 	bool LRUEvictImmediately();
 	void LRUTrim(uint32 nSpaceNeeded);
-	void LRUAdd(FVulkanRHIGraphicsPipelineState* PSO);
-	void LRUTouch(FVulkanRHIGraphicsPipelineState* PSO);
+	void LRUAdd(FVulkanGraphicsPipelineState* PSO);
+	void LRUTouch(FVulkanGraphicsPipelineState* PSO);
 	bool LRUEvictOne(bool bOnlyOld = false);
 	void LRURemoveAll();
 	void LRUDump();
 	void LRUDebugEvictAll(); //evict all that are safe to evict without stalling..
-	void LRURemove(FVulkanRHIGraphicsPipelineState* PSO);
-	void LRUCheckNotInside(FVulkanRHIGraphicsPipelineState* PSO);
+	void LRURemove(FVulkanGraphicsPipelineState* PSO);
+	void LRUCheckNotInside(FVulkanGraphicsPipelineState* PSO);
 
 
 	FVulkanDevice* Device;
@@ -593,11 +593,11 @@ private:
 
 
 	FCriticalSection GraphicsPSOLockedCS;
-	TMap<FVulkanPSOKey, FVulkanRHIGraphicsPipelineState*> GraphicsPSOLockedMap;
+	TMap<FVulkanPSOKey, FVulkanGraphicsPipelineState*> GraphicsPSOLockedMap;
 
 		
 	FCriticalSection LRUCS;
-	FVulkanRHIGraphicsPipelineStateLRU LRU;
+	FVulkanGraphicsPipelineStateLRU LRU;
 	uint32 LRUUsedPipelineSize = 0;
 	uint32 LRUUsedPipelineCount = 0;
 	uint32 LRUUsedPipelineMax = 0;
@@ -605,7 +605,7 @@ private:
 	bool bUseLRU = true;
 	friend class FVulkanDynamicRHI;
 	friend class FVulkanCommandListContext;
-	friend class FVulkanRHIGraphicsPipelineState;
+	friend class FVulkanGraphicsPipelineState;
 
 	struct FVulkanLRUCacheFile
 	{
@@ -665,7 +665,7 @@ public:
 		{
 			if (ImmediateDestroy)
 			{
-				VulkanRHI::vkDestroyPipeline(Device->GetInstanceHandle(), Pipeline, VULKAN_CPU_ALLOCATOR);
+				VulkanRHI::vkDestroyPipeline(Device->GetHandle(), Pipeline, VULKAN_CPU_ALLOCATOR);
 			}
 			else
 			{
@@ -684,7 +684,7 @@ protected:
 	FVulkanLayout* Layout; /*owned by FVulkanPipelineStateCacheManager, do not delete yourself !*/
 
 	friend class FVulkanPipelineStateCacheManager;
-	friend class FVulkanRHIGraphicsPipelineState;
+	friend class FVulkanGraphicsPipelineState;
 	friend class FVulkanComputePipelineDescriptorState;
 };
 
@@ -718,11 +718,11 @@ protected:
 	const bool bUsesBindless;
 };
 
-class FVulkanRHIGraphicsPipelineState : public FRHIGraphicsPipelineState
+class FVulkanGraphicsPipelineState : public FRHIGraphicsPipelineState
 {
 public:
-	FVulkanRHIGraphicsPipelineState(FVulkanDevice* Device, const FGraphicsPipelineStateInitializer& PSOInitializer, const FGfxPipelineDesc& Desc, FVulkanPSOKey* Key);
-	~FVulkanRHIGraphicsPipelineState();
+	FVulkanGraphicsPipelineState(FVulkanDevice* Device, const FGraphicsPipelineStateInitializer& PSOInitializer, const FGfxPipelineDesc& Desc, FVulkanPSOKey* Key);
+	~FVulkanGraphicsPipelineState();
 
 	inline const FVulkanVertexInputStateInfo& GetVertexInputState() const
 	{
@@ -801,7 +801,7 @@ public:
 	const FVulkanRenderPass* RenderPass;
 
 
-	FVulkanRHIGraphicsPipelineStateLRUNode* LRUNode = nullptr;
+	FVulkanGraphicsPipelineStateLRUNode* LRUNode = nullptr;
 	uint32 LRUFrame = UINT32_MAX;
 	uint32 PipelineCacheSize = UINT32_MAX;
 	uint64	PrecacheKey;// hash of elements relevant to the PSO cache
@@ -832,5 +832,5 @@ struct TVulkanResourceTraits<class FRHIComputePipelineState>
 template<>
 struct TVulkanResourceTraits<FRHIGraphicsPipelineState>
 {
-	typedef FVulkanRHIGraphicsPipelineState TConcreteType;
+	typedef FVulkanGraphicsPipelineState TConcreteType;
 };
diff --git a/Engine/Source/Runtime/VulkanRHI/Private/VulkanPipelineState.cpp b/Engine/Source/Runtime/VulkanRHI/Private/VulkanPipelineState.cpp
--- a/Engine/Source/Runtime/VulkanRHI/Private/VulkanPipelineState.cpp
+++ b/Engine/Source/Runtime/VulkanRHI/Private/VulkanPipelineState.cpp
@@ -224,7 +224,7 @@ bool FVulkanComputePipelineDescriptorState::InternalUpdateDescriptorSets(FVulkan
 		INC_DWORD_STAT(STAT_VulkanNumDescSets);
 		SCOPE_CYCLE_COUNTER(STAT_VulkanVkUpdateDS);
 #endif
-		VulkanRHI::vkUpdateDescriptorSets(Device.GetInstanceHandle(), DSWriteContainer.DescriptorWrites.Num(), DSWriteContainer.DescriptorWrites.GetData(), 0, nullptr);
+		VulkanRHI::vkUpdateDescriptorSets(Device.GetHandle(), DSWriteContainer.DescriptorWrites.Num(), DSWriteContainer.DescriptorWrites.GetData(), 0, nullptr);
 	}
 
 	return true;
@@ -281,7 +281,7 @@ void FVulkanComputePipelineDescriptorState::UpdateBindlessDescriptors(FVulkanCom
 		VkBufferDeviceAddressInfo BufferInfo;
 		ZeroVulkanStruct(BufferInfo, VK_STRUCTURE_TYPE_BUFFER_DEVICE_ADDRESS_INFO);
 		BufferInfo.buffer = WriteDescriptorSet.pBufferInfo->buffer;
-		VkDeviceAddress BufferAddress = VulkanRHI::vkGetBufferDeviceAddressKHR(Device.GetInstanceHandle(), &BufferInfo);
+		VkDeviceAddress BufferAddress = VulkanRHI::vkGetBufferDeviceAddressKHR(Device.GetHandle(), &BufferInfo);
 
 		DescriptorAddressInfo.sType = VK_STRUCTURE_TYPE_DESCRIPTOR_ADDRESS_INFO_EXT;
 		DescriptorAddressInfo.address = BufferAddress + WriteDescriptorSet.pBufferInfo->offset;
@@ -292,7 +292,7 @@ void FVulkanComputePipelineDescriptorState::UpdateBindlessDescriptors(FVulkanCom
 	Device.GetBindlessDescriptorManager()->RegisterUniformBuffers(Context, VK_PIPELINE_BIND_POINT_COMPUTE, StageUBs);
 }
 
-FVulkanGraphicsPipelineDescriptorState::FVulkanGraphicsPipelineDescriptorState(FVulkanDevice& InDevice, FVulkanRHIGraphicsPipelineState* InGfxPipeline)
+FVulkanGraphicsPipelineDescriptorState::FVulkanGraphicsPipelineDescriptorState(FVulkanDevice& InDevice, FVulkanGraphicsPipelineState* InGfxPipeline)
 	: FVulkanCommonPipelineDescriptorState(InDevice, ShaderStage::NumGraphicsStages, InGfxPipeline->UsesBindless())
 	, GfxPipeline(InGfxPipeline)
 {
@@ -454,7 +454,7 @@ bool FVulkanGraphicsPipelineDescriptorState::InternalUpdateDescriptorSets(FVulka
 			INC_DWORD_STAT_BY(STAT_VulkanNumDescSets, NumSets);
 			SCOPE_CYCLE_COUNTER(STAT_VulkanVkUpdateDS);
 	#endif
-			VulkanRHI::vkUpdateDescriptorSets(Device.GetInstanceHandle(), DSWriteContainer.DescriptorWrites.Num(), DSWriteContainer.DescriptorWrites.GetData(), 0, nullptr);
+			VulkanRHI::vkUpdateDescriptorSets(Device.GetHandle(), DSWriteContainer.DescriptorWrites.Num(), DSWriteContainer.DescriptorWrites.GetData(), 0, nullptr);
 
 			bIsResourcesDirty = false;
 		}
@@ -525,7 +525,7 @@ void FVulkanGraphicsPipelineDescriptorState::UpdateBindlessDescriptors(FVulkanCo
 				VkBufferDeviceAddressInfo BufferInfo;
 				ZeroVulkanStruct(BufferInfo, VK_STRUCTURE_TYPE_BUFFER_DEVICE_ADDRESS_INFO);
 				BufferInfo.buffer = WriteDescriptorSet.pBufferInfo->buffer;
-				VkDeviceAddress BufferAddress = VulkanRHI::vkGetBufferDeviceAddressKHR(Device.GetInstanceHandle(), &BufferInfo);
+				VkDeviceAddress BufferAddress = VulkanRHI::vkGetBufferDeviceAddressKHR(Device.GetHandle(), &BufferInfo);
 
 				DescriptorAddressInfo.sType = VK_STRUCTURE_TYPE_DESCRIPTOR_ADDRESS_INFO_EXT;
 				DescriptorAddressInfo.address = BufferAddress + WriteDescriptorSet.pBufferInfo->offset;
@@ -567,7 +567,7 @@ void FVulkanCommandListContext::ApplyStaticUniformBuffers(TRHIShader* Shader)
 
 void FVulkanCommandListContext::RHISetGraphicsPipelineState(FRHIGraphicsPipelineState* GraphicsState, uint32 StencilRef, bool bApplyAdditionalState)
 {
-	FVulkanRHIGraphicsPipelineState* Pipeline = ResourceCast(GraphicsState);
+	FVulkanGraphicsPipelineState* Pipeline = ResourceCast(GraphicsState);
 	
 	FVulkanPipelineStateCacheManager* PipelineStateCache = Device.GetPipelineStateCache();
 	PipelineStateCache->LRUTouch(Pipeline);
diff --git a/Engine/Source/Runtime/VulkanRHI/Private/VulkanPipelineState.h b/Engine/Source/Runtime/VulkanRHI/Private/VulkanPipelineState.h
--- a/Engine/Source/Runtime/VulkanRHI/Private/VulkanPipelineState.h
+++ b/Engine/Source/Runtime/VulkanRHI/Private/VulkanPipelineState.h
@@ -238,7 +238,7 @@ protected:
 class FVulkanGraphicsPipelineDescriptorState : public FVulkanCommonPipelineDescriptorState
 {
 public:
-	FVulkanGraphicsPipelineDescriptorState(FVulkanDevice& InDevice, FVulkanRHIGraphicsPipelineState* InGfxPipeline);
+	FVulkanGraphicsPipelineDescriptorState(FVulkanDevice& InDevice, FVulkanGraphicsPipelineState* InGfxPipeline);
 	virtual ~FVulkanGraphicsPipelineDescriptorState()
 	{
 		GfxPipeline->Release();
@@ -284,7 +284,7 @@ protected:
 	TStaticArray<uint32, ShaderStage::NumGraphicsStages> PackedUniformBuffersMask;
 	TStaticArray<uint32, ShaderStage::NumGraphicsStages> PackedUniformBuffersDirty;
 
-	FVulkanRHIGraphicsPipelineState* GfxPipeline;
+	FVulkanGraphicsPipelineState* GfxPipeline;
 
 	template<bool bUseDynamicGlobalUBs>
 	bool InternalUpdateDescriptorSets(FVulkanCommandListContext& Context);
diff --git a/Engine/Source/Runtime/VulkanRHI/Private/VulkanQuery.cpp b/Engine/Source/Runtime/VulkanRHI/Private/VulkanQuery.cpp
--- a/Engine/Source/Runtime/VulkanRHI/Private/VulkanQuery.cpp
+++ b/Engine/Source/Runtime/VulkanRHI/Private/VulkanQuery.cpp
@@ -55,13 +55,13 @@ FVulkanQueryPool::FVulkanQueryPool(FVulkanDevice& InDevice, uint32 InMaxQueries,
 	ZeroVulkanStruct(PoolCreateInfo, VK_STRUCTURE_TYPE_QUERY_POOL_CREATE_INFO);
 	PoolCreateInfo.queryType = GetVkQueryType(InQueryType);
 	PoolCreateInfo.queryCount = MaxQueries;
-	VERIFYVULKANRESULT(VulkanRHI::vkCreateQueryPool(Device.GetInstanceHandle(), &PoolCreateInfo, VULKAN_CPU_ALLOCATOR, &QueryPool));
+	VERIFYVULKANRESULT(VulkanRHI::vkCreateQueryPool(Device.GetHandle(), &PoolCreateInfo, VULKAN_CPU_ALLOCATOR, &QueryPool));
 }
 
 FVulkanQueryPool::~FVulkanQueryPool()
 {
 	DEC_DWORD_STAT(STAT_VulkanNumQueryPools);
-	VulkanRHI::vkDestroyQueryPool(Device.GetInstanceHandle(), QueryPool, VULKAN_CPU_ALLOCATOR);
+	VulkanRHI::vkDestroyQueryPool(Device.GetHandle(), QueryPool, VULKAN_CPU_ALLOCATOR);
 	QueryPool = VK_NULL_HANDLE;
 }
 
@@ -94,7 +94,7 @@ void FVulkanQueryPool::Reset(FVulkanCommandBuffer& InCmdBuffer)
 
 	if (Device.GetOptionalExtensions().HasEXTHostQueryReset)
 	{
-		VulkanRHI::vkResetQueryPoolEXT(GetDevice().GetInstanceHandle(), QueryPool, 0, MaxQueries);
+		VulkanRHI::vkResetQueryPoolEXT(GetDevice().GetHandle(), QueryPool, 0, MaxQueries);
 	}
 	else
 	{
diff --git a/Engine/Source/Runtime/VulkanRHI/Private/VulkanQueue.cpp b/Engine/Source/Runtime/VulkanRHI/Private/VulkanQueue.cpp
--- a/Engine/Source/Runtime/VulkanRHI/Private/VulkanQueue.cpp
+++ b/Engine/Source/Runtime/VulkanRHI/Private/VulkanQueue.cpp
@@ -65,7 +65,7 @@ FVulkanQueue::FVulkanQueue(FVulkanDevice& InDevice, uint32 InFamilyIndex, EVulka
 	, Device(InDevice)
 	, bUseTimelineSemaphores(InDevice.GetOptionalExtensions().HasKHRTimelineSemaphore&& GAllowTimelineSemaphores)
 {
-	VulkanRHI::vkGetDeviceQueue(Device.GetInstanceHandle(), FamilyIndex, QueueIndex, &Queue);
+	VulkanRHI::vkGetDeviceQueue(Device.GetHandle(), FamilyIndex, QueueIndex, &Queue);
 
 	FillSupportedStageBits();
 
@@ -366,7 +366,7 @@ void FVulkanQueue::Submit(TArrayView<VkSubmitInfo> InSubmitInfos, VulkanRHI::FFe
 
 	if (GWaitForIdleOnSubmit != 0)
 	{
-		VERIFYVULKANRESULT(VulkanRHI::vkDeviceWaitIdle(Device.GetInstanceHandle()));
+		VERIFYVULKANRESULT(VulkanRHI::vkDeviceWaitIdle(Device.GetHandle()));
 
 		if (Fence)
 		{
@@ -601,7 +601,7 @@ FVulkanDiagnosticBuffer::FVulkanDiagnosticBuffer(FVulkanDevice& InDevice, FVulka
 	ZeroVulkanStruct(CreateInfo, VK_STRUCTURE_TYPE_BUFFER_CREATE_INFO);
 	CreateInfo.size = SizeInBytes;
 	CreateInfo.usage = VK_BUFFER_USAGE_TRANSFER_DST_BIT;
-	VERIFYVULKANRESULT(VulkanRHI::vkCreateBuffer(Device.GetInstanceHandle(), &CreateInfo, VULKAN_CPU_ALLOCATOR, &Buffer));
+	VERIFYVULKANRESULT(VulkanRHI::vkCreateBuffer(Device.GetHandle(), &CreateInfo, VULKAN_CPU_ALLOCATOR, &Buffer));
 
 	const VulkanRHI::EVulkanAllocationFlags MemoryFlags =
 		VulkanRHI::EVulkanAllocationFlags::Dedicated |
@@ -627,7 +627,7 @@ FVulkanDiagnosticBuffer::~FVulkanDiagnosticBuffer()
 {
 	if (Buffer != VK_NULL_HANDLE)
 	{
-		VulkanRHI::vkDestroyBuffer(Device.GetInstanceHandle(), Buffer, VULKAN_CPU_ALLOCATOR);
+		VulkanRHI::vkDestroyBuffer(Device.GetHandle(), Buffer, VULKAN_CPU_ALLOCATOR);
 		Buffer = VK_NULL_HANDLE;
 	}
 
diff --git a/Engine/Source/Runtime/VulkanRHI/Private/VulkanRHI.cpp b/Engine/Source/Runtime/VulkanRHI/Private/VulkanRHI.cpp
--- a/Engine/Source/Runtime/VulkanRHI/Private/VulkanRHI.cpp
+++ b/Engine/Source/Runtime/VulkanRHI/Private/VulkanRHI.cpp
@@ -540,7 +540,7 @@ void FVulkanDynamicRHI::Shutdown()
 			for (auto& Pair : Device->SamplerMap)
 			{
 				FVulkanSamplerState* SamplerState = (FVulkanSamplerState*)Pair.Value.GetReference();
-				VulkanRHI::vkDestroySampler(Device->GetInstanceHandle(), SamplerState->Sampler, VULKAN_CPU_ALLOCATOR);
+				VulkanRHI::vkDestroySampler(Device->GetHandle(), SamplerState->Sampler, VULKAN_CPU_ALLOCATOR);
 			}
 			Device->SamplerMap.Empty();
 		}
@@ -1411,7 +1411,7 @@ VkDevice FVulkanDynamicRHI::RHIGetVkDevice() const
 {
 	if (Device)
 	{
-		return Device->GetInstanceHandle();
+		return Device->GetHandle();
 	}
 	return VK_NULL_HANDLE;
 }
@@ -1489,7 +1489,7 @@ bool FVulkanDynamicRHI::RHIDoesAdapterMatchDevice(const void* InAdapterId) const
 
 void* FVulkanDynamicRHI::RHIGetVkDeviceProcAddr(const char* InName) const
 {
-	return (void*)VulkanRHI::vkGetDeviceProcAddr(Device->GetInstanceHandle(), InName);
+	return (void*)VulkanRHI::vkGetDeviceProcAddr(Device->GetHandle(), InName);
 }
 
 void* FVulkanDynamicRHI::RHIGetVkInstanceProcAddr(const char* InName) const
@@ -1660,7 +1660,7 @@ void FVulkanDynamicRHI::RHIVerifyResult(VkResult Result, const ANSICHAR* VkFunti
 
 void* FVulkanDynamicRHI::RHIGetNativeDevice()
 {
-	return (void*)Device->GetInstanceHandle();
+	return (void*)Device->GetHandle();
 }
 
 void* FVulkanDynamicRHI::RHIGetNativePhysicalDevice()
@@ -1807,7 +1807,7 @@ void FVulkanDynamicRHI::SaveValidationCache()
 	VkValidationCacheEXT ValidationCache = GVulkanRHI->Device->GetValidationCache();
 	if (ValidationCache != VK_NULL_HANDLE)
 	{
-		VkDevice Device = GVulkanRHI->Device->GetInstanceHandle();
+		VkDevice Device = GVulkanRHI->Device->GetHandle();
 		PFN_vkGetValidationCacheDataEXT vkGetValidationCacheData = (PFN_vkGetValidationCacheDataEXT)(void*)VulkanRHI::vkGetDeviceProcAddr(Device, "vkGetValidationCacheDataEXT");
 		check(vkGetValidationCacheData);
 		size_t CacheSize = 0;
diff --git a/Engine/Source/Runtime/VulkanRHI/Private/VulkanRayTracing.cpp b/Engine/Source/Runtime/VulkanRHI/Private/VulkanRayTracing.cpp
--- a/Engine/Source/Runtime/VulkanRHI/Private/VulkanRayTracing.cpp
+++ b/Engine/Source/Runtime/VulkanRHI/Private/VulkanRayTracing.cpp
@@ -405,7 +405,7 @@ FVulkanRayTracingGeometry::FVulkanRayTracingGeometry(FRHICommandListBase& RHICmd
 
 	AccelerationStructureBuffer = ResourceCast(RHICmdList.CreateBuffer(CreateDesc).GetReference());
 
-	VkDevice NativeDevice = Device->GetInstanceHandle();
+	VkDevice NativeDevice = Device->GetHandle();
 
 	VkAccelerationStructureCreateInfoKHR CreateInfo;
 	ZeroVulkanStruct(CreateInfo, VK_STRUCTURE_TYPE_ACCELERATION_STRUCTURE_CREATE_INFO_KHR);
@@ -516,7 +516,7 @@ void FVulkanRayTracingGeometry::CompactAccelerationStructure(FVulkanCommandBuffe
 	INC_MEMORY_STAT_BY(STAT_VulkanRayTracingBLASMemory, AccelerationStructureBuffer->GetSize());
 	INC_MEMORY_STAT_BY(STAT_VulkanRayTracingStaticBLASMemory, AccelerationStructureBuffer->GetSize());
 
-	VkDevice NativeDevice = Device->GetInstanceHandle();
+	VkDevice NativeDevice = Device->GetHandle();
 
 	VkAccelerationStructureCreateInfoKHR CreateInfo;
 	ZeroVulkanStruct(CreateInfo, VK_STRUCTURE_TYPE_ACCELERATION_STRUCTURE_CREATE_INFO_KHR);
@@ -814,7 +814,7 @@ void BuildAccelerationStructure(
 	const VkDeviceAddress InstanceBufferAddress = InInstanceBuffer->GetDeviceAddress() + InInstanceOffset;
 
 	{
-		GetTLASBuildData(Scene.GetParent()->GetInstanceHandle(), NumInstances, InstanceBufferAddress, Scene.Initializer.BuildFlags, BuildMode, BuildData);
+		GetTLASBuildData(Scene.GetParent()->GetHandle(), NumInstances, InstanceBufferAddress, Scene.Initializer.BuildFlags, BuildMode, BuildData);
 
 		BuildData.GeometryInfo.dstAccelerationStructure = Scene.View->GetAccelerationStructureView().Handle;
 		BuildData.GeometryInfo.srcAccelerationStructure = bIsUpdate ? Scene.View->GetAccelerationStructureView().Handle : nullptr;
@@ -1026,7 +1026,7 @@ void FVulkanRayTracingShaderTable::Commit(FVulkanCommandListContext& Context, FR
 			{
 				ReleaseLocalBuffer(Device, Alloc);
 
-				const VkDevice DeviceHandle = Device->GetInstanceHandle();
+				const VkDevice DeviceHandle = Device->GetHandle();
 				const VkPhysicalDeviceRayTracingPipelinePropertiesKHR& RayTracingPipelineProps = Device->GetOptionalExtensionProperties().RayTracingPipelineProps;
 
 				// Fetch staging buffer and fill it
@@ -1090,7 +1090,7 @@ FRayTracingAccelerationStructureSize FVulkanDynamicRHI::RHICalcRayTracingSceneSi
 {
 	FVkRtTLASBuildData BuildData;
 	const VkDeviceAddress InstanceBufferAddress = 0; // No device address available when only querying TLAS size
-	GetTLASBuildData(Device->GetInstanceHandle(), Initializer.MaxNumInstances, InstanceBufferAddress, Initializer.BuildFlags, EAccelerationStructureBuildMode::Build, BuildData);
+	GetTLASBuildData(Device->GetHandle(), Initializer.MaxNumInstances, InstanceBufferAddress, Initializer.BuildFlags, EAccelerationStructureBuildMode::Build, BuildData);
 
 	FRayTracingAccelerationStructureSize Result;
 	Result.ResultSize = BuildData.SizesInfo.accelerationStructureSize;
@@ -1104,7 +1104,7 @@ FRayTracingAccelerationStructureSize FVulkanDynamicRHI::RHICalcRayTracingGeometr
 {	
 	FVkRtBLASBuildData BuildData;
 	GetBLASBuildData(
-		Device->GetInstanceHandle(),
+		Device->GetHandle(),
 		MakeArrayView(Initializer.Segments),
 		Initializer.GeometryType,
 		Initializer.IndexBuffer,
@@ -1227,7 +1227,7 @@ void FVulkanCommandListContext::RHIBuildAccelerationStructures(TConstArrayView<F
 
 		FVkRtBLASBuildData& BuildData = TempBuildData.AddDefaulted_GetRef();
 		GetBLASBuildData(
-			Device.GetInstanceHandle(),
+			Device.GetHandle(),
 			MakeArrayView(Geometry->Initializer.Segments),
 			Geometry->Initializer.GeometryType,
 			Geometry->Initializer.IndexBuffer,
@@ -1371,7 +1371,7 @@ public:
 	{
 		for (auto& Pair : PipelineLibraryMap)
 		{
-			VulkanRHI::vkDestroyPipeline(Device.GetInstanceHandle(), Pair.Value->PipelineHandle, VULKAN_CPU_ALLOCATOR);
+			VulkanRHI::vkDestroyPipeline(Device.GetHandle(), Pair.Value->PipelineHandle, VULKAN_CPU_ALLOCATOR);
 			Pair.Value->PipelineHandle = VK_NULL_HANDLE;
 		}
 		PipelineLibraryMap.Empty();
@@ -1652,7 +1652,7 @@ FVulkanRayTracingPipelineState::FVulkanRayTracingPipelineState(FVulkanDevice* co
 				const uint32 ShaderHandleStorageSize = HandleCount * HandleSize;
 				OutHandleStorage.AddUninitialized(ShaderHandleStorageSize);
 
-				VERIFYVULKANRESULT(VulkanDynamicAPI::vkGetRayTracingShaderGroupHandlesKHR(InDevice->GetInstanceHandle(), RTPipeline, HandleOffset, HandleCount, ShaderHandleStorageSize, OutHandleStorage.GetData()));
+				VERIFYVULKANRESULT(VulkanDynamicAPI::vkGetRayTracingShaderGroupHandlesKHR(InDevice->GetHandle(), RTPipeline, HandleOffset, HandleCount, ShaderHandleStorageSize, OutHandleStorage.GetData()));
 
 				HandleOffset += HandleCount;
 			}
@@ -1674,7 +1674,7 @@ FVulkanRayTracingPipelineState::~FVulkanRayTracingPipelineState()
 {
 	if (Pipeline != VK_NULL_HANDLE)
 	{
-		VulkanRHI::vkDestroyPipeline(Device->GetInstanceHandle(), Pipeline, VULKAN_CPU_ALLOCATOR);
+		VulkanRHI::vkDestroyPipeline(Device->GetHandle(), Pipeline, VULKAN_CPU_ALLOCATOR);
 		Pipeline = VK_NULL_HANDLE;
 	}
 }
@@ -1736,7 +1736,7 @@ FVulkanRayTracingCompactedSizeQueryPool::FVulkanRayTracingCompactedSizeQueryPool
 	: FVulkanQueryPool(InDevice, InMaxQueries, EVulkanQueryPoolType::ASCompactedSize)
 {
 	QueryOutput.SetNumZeroed(InMaxQueries);
-	VulkanRHI::vkResetQueryPoolEXT(InDevice.GetInstanceHandle(), QueryPool, 0, MaxQueries);
+	VulkanRHI::vkResetQueryPoolEXT(InDevice.GetHandle(), QueryPool, 0, MaxQueries);
 }
 
 void FVulkanRayTracingCompactedSizeQueryPool::EndBatch(FVulkanCommandListContext& CommandContext)
@@ -1759,7 +1759,7 @@ bool FVulkanRayTracingCompactedSizeQueryPool::TryGetResults(uint32 NumResults)
 		return false;
 	}
 
-	VkResult Result = VulkanRHI::vkGetQueryPoolResults(Device.GetInstanceHandle(), QueryPool, 0, NumResults, NumResults * sizeof(uint64), QueryOutput.GetData(), sizeof(uint64), VK_QUERY_RESULT_64_BIT | VK_QUERY_RESULT_WAIT_BIT);
+	VkResult Result = VulkanRHI::vkGetQueryPoolResults(Device.GetHandle(), QueryPool, 0, NumResults, NumResults * sizeof(uint64), QueryOutput.GetData(), sizeof(uint64), VK_QUERY_RESULT_64_BIT | VK_QUERY_RESULT_WAIT_BIT);
 	if (Result == VK_SUCCESS)
 	{
 		return true;
diff --git a/Engine/Source/Runtime/VulkanRHI/Private/VulkanRenderpass.h b/Engine/Source/Runtime/VulkanRHI/Private/VulkanRenderpass.h
--- a/Engine/Source/Runtime/VulkanRHI/Private/VulkanRenderpass.h
+++ b/Engine/Source/Runtime/VulkanRHI/Private/VulkanRenderpass.h
@@ -356,7 +356,7 @@ struct FVulkanRenderPassCreateInfo<VkRenderPassCreateInfo>
 	VkRenderPass Create(FVulkanDevice& Device)
 	{
 		VkRenderPass Handle = VK_NULL_HANDLE;
-		VERIFYVULKANRESULT_EXPANDED(VulkanRHI::vkCreateRenderPass(Device.GetInstanceHandle(), this, VULKAN_CPU_ALLOCATOR, &Handle));
+		VERIFYVULKANRESULT_EXPANDED(VulkanRHI::vkCreateRenderPass(Device.GetHandle(), this, VULKAN_CPU_ALLOCATOR, &Handle));
 		return Handle;
 	}
 };
@@ -397,7 +397,7 @@ struct FVulkanRenderPassCreateInfo<VkRenderPassCreateInfo2>
 	VkRenderPass Create(FVulkanDevice& Device)
 	{
 		VkRenderPass Handle = VK_NULL_HANDLE;
-		VERIFYVULKANRESULT_EXPANDED(VulkanRHI::vkCreateRenderPass2KHR(Device.GetInstanceHandle(), this, VULKAN_CPU_ALLOCATOR, &Handle));
+		VERIFYVULKANRESULT_EXPANDED(VulkanRHI::vkCreateRenderPass2KHR(Device.GetHandle(), this, VULKAN_CPU_ALLOCATOR, &Handle));
 		return Handle;
 	}
 };
diff --git a/Engine/Source/Runtime/VulkanRHI/Private/VulkanShaders.cpp b/Engine/Source/Runtime/VulkanRHI/Private/VulkanShaders.cpp
--- a/Engine/Source/Runtime/VulkanRHI/Private/VulkanShaders.cpp
+++ b/Engine/Source/Runtime/VulkanRHI/Private/VulkanShaders.cpp
@@ -328,7 +328,7 @@ static TRefCountPtr<FVulkanShaderModule> CreateShaderModule(FVulkanDevice* Devic
 	}
 #endif
 
-	VERIFYVULKANRESULT(VulkanRHI::vkCreateShaderModule(Device->GetInstanceHandle(), &ModuleCreateInfo, VULKAN_CPU_ALLOCATOR, &ShaderModule));
+	VERIFYVULKANRESULT(VulkanRHI::vkCreateShaderModule(Device->GetHandle(), &ModuleCreateInfo, VULKAN_CPU_ALLOCATOR, &ShaderModule));
 	
 	TRefCountPtr<FVulkanShaderModule> ReturnPtr = TRefCountPtr<FVulkanShaderModule>(new FVulkanShaderModule(Device, ShaderModule));
 
diff --git a/Engine/Source/Runtime/VulkanRHI/Private/VulkanState.cpp b/Engine/Source/Runtime/VulkanRHI/Private/VulkanState.cpp
--- a/Engine/Source/Runtime/VulkanRHI/Private/VulkanState.cpp
+++ b/Engine/Source/Runtime/VulkanRHI/Private/VulkanState.cpp
@@ -227,7 +227,7 @@ FVulkanSamplerState::FVulkanSamplerState(const VkSamplerCreateInfo& InInfo, FVul
 	, BindlessHandle()
 	, bIsImmutable(bInIsImmutable)
 {
-	VERIFYVULKANRESULT(VulkanRHI::vkCreateSampler(InDevice.GetInstanceHandle(), &InInfo, VULKAN_CPU_ALLOCATOR, &Sampler));
+	VERIFYVULKANRESULT(VulkanRHI::vkCreateSampler(InDevice.GetHandle(), &InInfo, VULKAN_CPU_ALLOCATOR, &Sampler));
 
 	if (UseVulkanDescriptorCache())
 	{
diff --git a/Engine/Source/Runtime/VulkanRHI/Private/VulkanSubmission.cpp b/Engine/Source/Runtime/VulkanRHI/Private/VulkanSubmission.cpp
--- a/Engine/Source/Runtime/VulkanRHI/Private/VulkanSubmission.cpp
+++ b/Engine/Source/Runtime/VulkanRHI/Private/VulkanSubmission.cpp
@@ -207,7 +207,7 @@ void FVulkanDynamicRHI::KickInterruptThread()
 		ZeroVulkanStruct(SignalInfo, VK_STRUCTURE_TYPE_SEMAPHORE_SIGNAL_INFO);
 		SignalInfo.semaphore = CPUTimelineSemaphore->GetHandle();
 		SignalInfo.value = (uint64_t)CPUTimelineSemaphoreValue++;
-		VERIFYVULKANRESULT(VulkanRHI::vkSignalSemaphoreKHR(Device->GetInstanceHandle(), &SignalInfo));
+		VERIFYVULKANRESULT(VulkanRHI::vkSignalSemaphoreKHR(Device->GetHandle(), &SignalInfo));
 
 		InterruptThread->Kick();
 	}
@@ -452,7 +452,7 @@ bool FVulkanDynamicRHI::WaitAndProcessInterruptQueue()
 				WaitInfo.semaphoreCount = NumWaits;
 				WaitInfo.pSemaphores = Semaphores;
 				WaitInfo.pValues = Values;
-				VERIFYVULKANRESULT_EXPANDED(VulkanRHI::vkWaitSemaphoresKHR(Device->GetInstanceHandle(), &WaitInfo, (uint64_t)Timeout));
+				VERIFYVULKANRESULT_EXPANDED(VulkanRHI::vkWaitSemaphoresKHR(Device->GetHandle(), &WaitInfo, (uint64_t)Timeout));
 			}
 
 			// Some work completed and we made progress, or we stopped waiting because of CPUTimelineSemaphore
@@ -661,7 +661,7 @@ void FVulkanDynamicRHI::CompletePayload(FVulkanPayload* Payload)
 			{
 				const VkQueryResultFlags QueryFlags = VK_QUERY_RESULT_64_BIT | VK_QUERY_RESULT_WITH_AVAILABILITY_BIT;
 				Result = VulkanRHI::vkGetQueryPoolResults(
-					Device->GetInstanceHandle(),
+					Device->GetHandle(),
 					QueryPool->GetHandle(),
 					0,
 					QueryPool->CurrentQueryCount,
diff --git a/Engine/Source/Runtime/VulkanRHI/Private/VulkanSwapChain.cpp b/Engine/Source/Runtime/VulkanRHI/Private/VulkanSwapChain.cpp
--- a/Engine/Source/Runtime/VulkanRHI/Private/VulkanSwapChain.cpp
+++ b/Engine/Source/Runtime/VulkanRHI/Private/VulkanSwapChain.cpp
@@ -545,7 +545,7 @@ FVulkanSwapChain::FVulkanSwapChain(VkInstance InInstance, FVulkanDevice& InDevic
 		Swap(SwapChainInfo.imageExtent.width, SwapChainInfo.imageExtent.height);
 	}
 
-	VkResult Result = FVulkanPlatform::CreateSwapchainKHR(WindowContext, Device.GetPhysicalHandle(), Device.GetInstanceHandle(), &SwapChainInfo, VULKAN_CPU_ALLOCATOR, &SwapChain);
+	VkResult Result = FVulkanPlatform::CreateSwapchainKHR(WindowContext, Device.GetPhysicalHandle(), Device.GetHandle(), &SwapChainInfo, VULKAN_CPU_ALLOCATOR, &SwapChain);
 
 	// Swap the aspect ratio back to make it transparent to high level
 	if (bSwapChainNeedsTransform)
@@ -559,7 +559,7 @@ FVulkanSwapChain::FVulkanSwapChain(VkInstance InInstance, FVulkanDevice& InDevic
 		// Unlink fullscreen
 		UE_LOG(LogVulkanRHI, Warning, TEXT("Create swapchain failed with Initialization error; removing FullScreen extension..."));
 		SwapChainInfo.pNext = FullScreenInfo.pNext;
-		Result = FVulkanPlatform::CreateSwapchainKHR(WindowContext, Device.GetPhysicalHandle(), Device.GetInstanceHandle(), &SwapChainInfo, VULKAN_CPU_ALLOCATOR, &SwapChain);
+		Result = FVulkanPlatform::CreateSwapchainKHR(WindowContext, Device.GetPhysicalHandle(), Device.GetHandle(), &SwapChainInfo, VULKAN_CPU_ALLOCATOR, &SwapChain);
 	}
 #endif
 	VERIFYVULKANRESULT_EXPANDED(Result);
@@ -568,7 +568,7 @@ FVulkanSwapChain::FVulkanSwapChain(VkInstance InInstance, FVulkanDevice& InDevic
 	{
 		if (RecreateInfo->SwapChain != VK_NULL_HANDLE)
 		{
-			FVulkanPlatform::DestroySwapchainKHR(Device.GetInstanceHandle(), RecreateInfo->SwapChain, VULKAN_CPU_ALLOCATOR);
+			FVulkanPlatform::DestroySwapchainKHR(Device.GetHandle(), RecreateInfo->SwapChain, VULKAN_CPU_ALLOCATOR);
 			RecreateInfo->SwapChain = VK_NULL_HANDLE;
 		}
 		if (RecreateInfo->Surface != VK_NULL_HANDLE)
@@ -583,7 +583,7 @@ FVulkanSwapChain::FVulkanSwapChain(VkInstance InInstance, FVulkanDevice& InDevic
 	bInternalFullScreen = bIsFullScreen;
 
 	uint32 NumSwapChainImages;
-	VERIFYVULKANRESULT_EXPANDED(VulkanRHI::vkGetSwapchainImagesKHR(Device.GetInstanceHandle(), SwapChain, &NumSwapChainImages, nullptr));
+	VERIFYVULKANRESULT_EXPANDED(VulkanRHI::vkGetSwapchainImagesKHR(Device.GetHandle(), SwapChain, &NumSwapChainImages, nullptr));
 
 	if (GVulkanSwapChainIgnoreExtraImages != 0)
 	{
@@ -591,7 +591,7 @@ FVulkanSwapChain::FVulkanSwapChain(VkInstance InInstance, FVulkanDevice& InDevic
 	}
 
 	OutImages.AddUninitialized(NumSwapChainImages);
-	VERIFYVULKANRESULT_EXPANDED(VulkanRHI::vkGetSwapchainImagesKHR(Device.GetInstanceHandle(), SwapChain, &NumSwapChainImages, OutImages.GetData()));
+	VERIFYVULKANRESULT_EXPANDED(VulkanRHI::vkGetSwapchainImagesKHR(Device.GetHandle(), SwapChain, &NumSwapChainImages, OutImages.GetData()));
 
 #if VULKAN_USE_IMAGE_ACQUIRE_FENCES
 	ImageAcquiredFences.AddUninitialized(NumSwapChainImages);
@@ -627,7 +627,7 @@ void FVulkanSwapChain::Destroy(FVulkanSwapChainRecreateInfo* RecreateInfo)
 	}
 	else
 	{
-		FVulkanPlatform::DestroySwapchainKHR(Device.GetInstanceHandle(), SwapChain, VULKAN_CPU_ALLOCATOR);
+		FVulkanPlatform::DestroySwapchainKHR(Device.GetHandle(), SwapChain, VULKAN_CPU_ALLOCATOR);
 	}
 	SwapChain = VK_NULL_HANDLE;
 
@@ -681,7 +681,7 @@ int32 FVulkanSwapChain::AcquireImageIndex(VulkanRHI::FSemaphore** OutSemaphore)
 		FRenderThreadIdleScope IdleScope(ERenderThreadIdleTypes::WaitingForGPUPresent);
 
 		Result = VulkanRHI::vkAcquireNextImageKHR(
-			Device.GetInstanceHandle(),
+			Device.GetHandle(),
 			SwapChain,
 			UINT64_MAX,
 			ImageAcquiredSemaphore[SemaphoreIndex]->GetHandle(),
@@ -692,7 +692,7 @@ int32 FVulkanSwapChain::AcquireImageIndex(VulkanRHI::FSemaphore** OutSemaphore)
 		while (ImageIndex > MaxImageIndex && (Result == VK_SUCCESS || Result == VK_SUBOPTIMAL_KHR))
 		{
 			Result = VulkanRHI::vkAcquireNextImageKHR(
-				Device.GetInstanceHandle(),
+				Device.GetHandle(),
 				SwapChain,
 				UINT64_MAX,
 				ImageAcquiredSemaphore[SemaphoreIndex]->GetHandle(),
diff --git a/Engine/Source/Runtime/VulkanRHI/Private/VulkanTexture.cpp b/Engine/Source/Runtime/VulkanRHI/Private/VulkanTexture.cpp
--- a/Engine/Source/Runtime/VulkanRHI/Private/VulkanTexture.cpp
+++ b/Engine/Source/Runtime/VulkanRHI/Private/VulkanTexture.cpp
@@ -586,7 +586,7 @@ void FVulkanTexture::InternalMoveSurface(FVulkanDevice& InDevice, FVulkanCommand
 	FVulkanTexture::GenerateImageCreateInfo(ImageCreateInfo, InDevice, Desc, &StorageFormat, &ViewFormat);
 
 	VkImage MovedImage;
-	VERIFYVULKANRESULT(VulkanRHI::vkCreateImage(InDevice.GetInstanceHandle(), &ImageCreateInfo.ImageCreateInfo, VULKAN_CPU_ALLOCATOR, &MovedImage));
+	VERIFYVULKANRESULT(VulkanRHI::vkCreateImage(InDevice.GetHandle(), &ImageCreateInfo.ImageCreateInfo, VULKAN_CPU_ALLOCATOR, &MovedImage));
 	checkf(Tiling == ImageCreateInfo.ImageCreateInfo.tiling, TEXT("Move has changed image tiling:  before [%s] != after [%s]"), VK_TYPE_TO_STRING(VkImageTiling, Tiling), VK_TYPE_TO_STRING(VkImageTiling, ImageCreateInfo.ImageCreateInfo.tiling));
 
 	const ETextureCreateFlags UEFlags = Desc.Flags;
@@ -601,7 +601,7 @@ void FVulkanTexture::InternalMoveSurface(FVulkanDevice& InDevice, FVulkanCommand
 #if UE_BUILD_DEBUG || UE_BUILD_DEVELOPMENT
 	// This shouldn't change
 	VkMemoryRequirements MovedMemReqs;
-	VulkanRHI::vkGetImageMemoryRequirements(InDevice.GetInstanceHandle(), MovedImage, &MovedMemReqs);
+	VulkanRHI::vkGetImageMemoryRequirements(InDevice.GetHandle(), MovedImage, &MovedMemReqs);
 	checkf((MemoryRequirements.alignment == MovedMemReqs.alignment), TEXT("Memory requirements changed: alignment %d -> %d"), (int32)MemoryRequirements.alignment, (int32)MovedMemReqs.alignment);
 	checkf((MemoryRequirements.size == MovedMemReqs.size), TEXT("Memory requirements changed: size %d -> %d"), (int32)MemoryRequirements.size, (int32)MovedMemReqs.size);
 	checkf((MemoryRequirements.memoryTypeBits == MovedMemReqs.memoryTypeBits), TEXT("Memory requirements changed: memoryTypeBits %d -> %d"), (int32)MemoryRequirements.memoryTypeBits, (int32)MovedMemReqs.memoryTypeBits);
@@ -1605,7 +1605,7 @@ FVulkanTexture::FVulkanTexture(FVulkanDevice& InDevice, const FRHITextureCreateD
 		ViewFormat = StorageFormat = UEToVkTextureFormat(InCreateDesc.Format, false);
 
 		// :todo-jn: Kept around temporarily for legacy defrag/eviction/stats
-		VulkanRHI::vkGetBufferMemoryRequirements(InDevice.GetInstanceHandle(), CpuReadbackBuffer->Buffer, &MemoryRequirements);
+		VulkanRHI::vkGetBufferMemoryRequirements(InDevice.GetHandle(), CpuReadbackBuffer->Buffer, &MemoryRequirements);
 
 		return;
 	}
@@ -1613,10 +1613,10 @@ FVulkanTexture::FVulkanTexture(FVulkanDevice& InDevice, const FRHITextureCreateD
 	FImageCreateInfo ImageCreateInfo;
 	FVulkanTexture::GenerateImageCreateInfo(ImageCreateInfo, InDevice, InCreateDesc, &StorageFormat, &ViewFormat);
 
-	VERIFYVULKANRESULT(VulkanRHI::vkCreateImage(InDevice.GetInstanceHandle(), &ImageCreateInfo.ImageCreateInfo, VULKAN_CPU_ALLOCATOR, &Image));
+	VERIFYVULKANRESULT(VulkanRHI::vkCreateImage(InDevice.GetHandle(), &ImageCreateInfo.ImageCreateInfo, VULKAN_CPU_ALLOCATOR, &Image));
 
 	// Fetch image size
-	VulkanRHI::vkGetImageMemoryRequirements(InDevice.GetInstanceHandle(), Image, &MemoryRequirements);
+	VulkanRHI::vkGetImageMemoryRequirements(InDevice.GetHandle(), Image, &MemoryRequirements);
 
 	VULKAN_SET_DEBUG_NAME(InDevice, VK_OBJECT_TYPE_IMAGE, Image, TEXT("%s:(FVulkanTexture*)0x%p"), InCreateDesc.DebugName ? InCreateDesc.DebugName : TEXT("?"), this);
 
@@ -1733,7 +1733,7 @@ FVulkanTexture::FVulkanTexture(FVulkanDevice& InDevice, const FRHITextureCreateD
 		{
 			ImageUsageFlags = GetUsageFlagsFromCreateFlags(InDevice, InCreateDesc.Flags);
 #if VULKAN_ENABLE_WRAP_LAYER
-			FWrapLayer::CreateImage(VK_SUCCESS, InDevice.GetInstanceHandle(), nullptr, &Image);
+			FWrapLayer::CreateImage(VK_SUCCESS, InDevice.GetHandle(), nullptr, &Image);
 #endif
 			VULKAN_SET_DEBUG_NAME(InDevice, VK_OBJECT_TYPE_IMAGE, Image, TEXT("%s:(FVulkanTexture*)0x%p"), InCreateDesc.DebugName ? InCreateDesc.DebugName : TEXT("?"), this);
 
@@ -1807,7 +1807,7 @@ FVulkanTexture::FVulkanTexture(FVulkanDevice& InDevice, const FRHITextureCreateD
 	AHardwareBuffer_acquire(HardwareBuffer);
 
 	IVulkanDynamicRHI* RHI = GetIVulkanDynamicRHI();
-	VkDevice VulkanDevice = InDevice.GetInstanceHandle();
+	VkDevice VulkanDevice = InDevice.GetHandle();
 	const VkAllocationCallbacks* AllocationCallbacks = RHI->RHIGetVkAllocationCallbacks();
 
 	VkAndroidHardwareBufferFormatPropertiesANDROID HardwareBufferFormatProperties;
@@ -1933,7 +1933,7 @@ FVulkanTexture::FVulkanTexture(FVulkanDevice& InDevice, const FRHITextureCreateD
 		{
 			ImageUsageFlags = GetUsageFlagsFromCreateFlags(InDevice, InCreateDesc.Flags);
 #if VULKAN_ENABLE_WRAP_LAYER
-			FWrapLayer::CreateImage(VK_SUCCESS, InDevice.GetInstanceHandle(), nullptr, &Image);
+			FWrapLayer::CreateImage(VK_SUCCESS, InDevice.GetHandle(), nullptr, &Image);
 #endif
 			VULKAN_SET_DEBUG_NAME(InDevice, VK_OBJECT_TYPE_IMAGE, Image, TEXT("%s:(FVulkanTexture*)0x%p"), InCreateDesc.DebugName ? InCreateDesc.DebugName : TEXT("?"), this);
 
@@ -2251,15 +2251,15 @@ FDynamicRHI::FRHICalcTextureSizeResult FVulkanDynamicRHI::RHICalcTexturePlatform
 		VkMemoryRequirements2 MemReq2;
 		ZeroVulkanStruct(MemReq2, VK_STRUCTURE_TYPE_MEMORY_REQUIREMENTS_2);
 
-		VulkanRHI::vkGetDeviceImageMemoryRequirementsKHR(Device->GetInstanceHandle(), &ImageMemReq, &MemReq2);
+		VulkanRHI::vkGetDeviceImageMemoryRequirementsKHR(Device->GetHandle(), &ImageMemReq, &MemReq2);
 		OutMemReq = MemReq2.memoryRequirements;
 	}
 	else
 	{
 		VkImage TmpImage;
-		VERIFYVULKANRESULT(VulkanRHI::vkCreateImage(Device->GetInstanceHandle(), &TmpCreateInfo.ImageCreateInfo, VULKAN_CPU_ALLOCATOR, &TmpImage));
-		VulkanRHI::vkGetImageMemoryRequirements(Device->GetInstanceHandle(), TmpImage, &OutMemReq);
-		VulkanRHI::vkDestroyImage(Device->GetInstanceHandle(), TmpImage, VULKAN_CPU_ALLOCATOR);
+		VERIFYVULKANRESULT(VulkanRHI::vkCreateImage(Device->GetHandle(), &TmpCreateInfo.ImageCreateInfo, VULKAN_CPU_ALLOCATOR, &TmpImage));
+		VulkanRHI::vkGetImageMemoryRequirements(Device->GetHandle(), TmpImage, &OutMemReq);
+		VulkanRHI::vkDestroyImage(Device->GetHandle(), TmpImage, VULKAN_CPU_ALLOCATOR);
 	}
 
 	{
diff --git a/Engine/Source/Runtime/VulkanRHI/Private/VulkanUniformBuffer.cpp b/Engine/Source/Runtime/VulkanRHI/Private/VulkanUniformBuffer.cpp
--- a/Engine/Source/Runtime/VulkanRHI/Private/VulkanUniformBuffer.cpp
+++ b/Engine/Source/Runtime/VulkanRHI/Private/VulkanUniformBuffer.cpp
@@ -261,7 +261,7 @@ VkDeviceAddress FVulkanUniformBuffer::GetDeviceAddress() const
 	VkBufferDeviceAddressInfo BufferInfo;
 	ZeroVulkanStruct(BufferInfo, VK_STRUCTURE_TYPE_BUFFER_DEVICE_ADDRESS_INFO);
 	BufferInfo.buffer = GetBufferHandle();
-	VkDeviceAddress BufferAddress = VulkanRHI::vkGetBufferDeviceAddressKHR(Device->GetInstanceHandle(), &BufferInfo);
+	VkDeviceAddress BufferAddress = VulkanRHI::vkGetBufferDeviceAddressKHR(Device->GetHandle(), &BufferInfo);
 	BufferAddress += GetOffset();
 	return BufferAddress;
 }
diff --git a/Engine/Source/Runtime/VulkanRHI/Private/VulkanUtil.cpp b/Engine/Source/Runtime/VulkanRHI/Private/VulkanUtil.cpp
--- a/Engine/Source/Runtime/VulkanRHI/Private/VulkanUtil.cpp
+++ b/Engine/Source/Runtime/VulkanRHI/Private/VulkanUtil.cpp
@@ -48,7 +48,7 @@ void SetVulkanResourceName(FVulkanDevice* Device, FVulkanTexture* Texture, const
 	if (PFN_vkSetDebugUtilsObjectNameEXT SetDebugName = Device->GetSetDebugName())
 	{
 		FAnsiString AnsiName(Name);
-		VulkanRHI::SetDebugName(SetDebugName, Device->GetInstanceHandle(), Texture->Image, GetData(AnsiName));
+		VulkanRHI::SetDebugName(SetDebugName, Device->GetHandle(), Texture->Image, GetData(AnsiName));
 	}
 #endif
 }
@@ -470,7 +470,7 @@ void AftermathLateAssociate(float TimeLimitSeconds, uint32 FrameLimit)
 				checkNoEntry();
 			case RRT_GraphicsPipelineState:
 			{
-				auto* Pipeline = static_cast<FVulkanRHIGraphicsPipelineState*>(Resource);
+				auto* Pipeline = static_cast<FVulkanGraphicsPipelineState*>(Resource);
 					
 				if (IsPipelineUnused(Pipeline->FrameCounter, FrameLimit))
 				{
@@ -630,8 +630,8 @@ FVulkanGPUFence::FVulkanGPUFence(FVulkanDevice& InDevice, FName InName)
 	{
 		VkEventCreateInfo EventCreateInfo;
 		ZeroVulkanStruct(EventCreateInfo, VK_STRUCTURE_TYPE_EVENT_CREATE_INFO);
-		VERIFYVULKANRESULT(VulkanRHI::vkCreateEvent(InDevice.GetInstanceHandle(), &EventCreateInfo, VULKAN_CPU_ALLOCATOR, &Event));
-		VERIFYVULKANRESULT(VulkanRHI::vkResetEvent(InDevice.GetInstanceHandle(), Event));
+		VERIFYVULKANRESULT(VulkanRHI::vkCreateEvent(InDevice.GetHandle(), &EventCreateInfo, VULKAN_CPU_ALLOCATOR, &Event));
+		VERIFYVULKANRESULT(VulkanRHI::vkResetEvent(InDevice.GetHandle(), Event));
 	}
 }
 
@@ -650,7 +650,7 @@ void FVulkanGPUFence::Clear()
 
 	if (Event)
 	{
-		VERIFYVULKANRESULT(VulkanRHI::vkResetEvent(Device.GetInstanceHandle(), Event));
+		VERIFYVULKANRESULT(VulkanRHI::vkResetEvent(Device.GetHandle(), Event));
 		SubmittedSyncPoint = nullptr;
 	}
 }
@@ -664,7 +664,7 @@ bool FVulkanGPUFence::Poll() const
 
 	if (Event && SubmittedSyncPoint.IsValid() && SubmittedSyncPoint->IsComplete())
 	{
-		return (VulkanRHI::vkGetEventStatus(Device.GetInstanceHandle(), Event) == VK_EVENT_SET);
+		return (VulkanRHI::vkGetEventStatus(Device.GetHandle(), Event) == VK_EVENT_SET);
 	}
 
 	return false;
@@ -721,7 +721,7 @@ namespace VulkanRHI
 {
 	VkBuffer CreateBuffer(FVulkanDevice* InDevice, VkDeviceSize Size, VkBufferUsageFlags BufferUsageFlags, VkMemoryRequirements& OutMemoryRequirements)
 	{
-		VkDevice Device = InDevice->GetInstanceHandle();
+		VkDevice Device = InDevice->GetHandle();
 		VkBuffer Buffer = VK_NULL_HANDLE;
 
 		VkBufferCreateInfo BufferCreateInfo;
@@ -879,7 +879,7 @@ void FVulkanDynamicRHI::DumpCrashMarkers()
 
 	if (Device->GetOptionalExtensions().HasEXTDeviceFault)
 	{
-		const VkDevice DeviceHandle = Device->GetInstanceHandle();
+		const VkDevice DeviceHandle = Device->GetHandle();
 			VkResult Result;
 
 			VkDeviceFaultCountsEXT FaultCounts;
diff --git a/Engine/Source/Runtime/VulkanRHI/Private/VulkanView.cpp b/Engine/Source/Runtime/VulkanRHI/Private/VulkanView.cpp
--- a/Engine/Source/Runtime/VulkanRHI/Private/VulkanView.cpp
+++ b/Engine/Source/Runtime/VulkanRHI/Private/VulkanView.cpp
@@ -173,7 +173,7 @@ FVulkanView* FVulkanView::InitAsTypedBufferView(FVulkanBuffer* Buffer, EPixelFor
 	check(Buffer->GetBufferUsageFlags() & (VK_BUFFER_USAGE_UNIFORM_TEXEL_BUFFER_BIT | VK_BUFFER_USAGE_STORAGE_TEXEL_BUFFER_BIT));
 	check(IsAligned(InOffset, Limits.minTexelBufferOffsetAlignment));
 
-	VERIFYVULKANRESULT(VulkanRHI::vkCreateBufferView(Device.GetInstanceHandle(), &ViewInfo, VULKAN_CPU_ALLOCATOR, &TBV.View));
+	VERIFYVULKANRESULT(VulkanRHI::vkCreateBufferView(Device.GetHandle(), &ViewInfo, VULKAN_CPU_ALLOCATOR, &TBV.View));
 
 	TBV.bVolatile = Buffer->IsVolatile();
 	if (!TBV.bVolatile && UseVulkanDescriptorCache())
@@ -269,7 +269,7 @@ FVulkanView* FVulkanView::InitAsTextureView(VkImage InImage, const FVulkanTextur
 	}
 
 	INC_DWORD_STAT(STAT_VulkanNumImageViews);
-	VERIFYVULKANRESULT(VulkanRHI::vkCreateImageView(Device.GetInstanceHandle(), &ViewInfo, VULKAN_CPU_ALLOCATOR, &TV.View));
+	VERIFYVULKANRESULT(VulkanRHI::vkCreateImageView(Device.GetHandle(), &ViewInfo, VULKAN_CPU_ALLOCATOR, &TV.View));
 
 	TV.Image = InImage;
 
@@ -325,7 +325,7 @@ FVulkanView* FVulkanView::InitAsAccelerationStructureView(FVulkanBuffer* Buffer,
 	CreateInfo.size = Size;
 	CreateInfo.type = VK_ACCELERATION_STRUCTURE_TYPE_TOP_LEVEL_KHR;
 
-	VERIFYVULKANRESULT(VulkanDynamicAPI::vkCreateAccelerationStructureKHR(Device.GetInstanceHandle(), &CreateInfo, VULKAN_CPU_ALLOCATOR, &ASV.Handle));
+	VERIFYVULKANRESULT(VulkanDynamicAPI::vkCreateAccelerationStructureKHR(Device.GetHandle(), &CreateInfo, VULKAN_CPU_ALLOCATOR, &ASV.Handle));
 
 	Device.GetBindlessDescriptorManager()->UpdateAccelerationStructure(BindlessHandle, ASV.Handle);
 
diff --git a/Engine/Source/Runtime/VulkanRHI/Private/VulkanViewport.cpp b/Engine/Source/Runtime/VulkanRHI/Private/VulkanViewport.cpp
--- a/Engine/Source/Runtime/VulkanRHI/Private/VulkanViewport.cpp
+++ b/Engine/Source/Runtime/VulkanRHI/Private/VulkanViewport.cpp
@@ -393,7 +393,7 @@ void FVulkanViewport::InitImages(FVulkanContextCommon& Context, TConstArrayView<
 #if VULKAN_ENABLE_DRAW_MARKERS
 		if (Device->GetSetDebugName())
 		{
-			VulkanRHI::SetDebugName(Device->GetSetDebugName(), Device->GetInstanceHandle(), BackBufferImages[Index]->Image, "VulkanBackBuffer");
+			VulkanRHI::SetDebugName(Device->GetSetDebugName(), Device->GetHandle(), BackBufferImages[Index]->Image, "VulkanBackBuffer");
 		}
 #endif
 	}
@@ -453,7 +453,7 @@ void FVulkanViewport::CreateSwapchain(FVulkanCommandListContext& Context, FVulka
 		{
 			if (RecreateInfo->SwapChain)
 			{
-				FVulkanPlatform::DestroySwapchainKHR(Device->GetInstanceHandle(), RecreateInfo->SwapChain, VULKAN_CPU_ALLOCATOR);
+				FVulkanPlatform::DestroySwapchainKHR(Device->GetHandle(), RecreateInfo->SwapChain, VULKAN_CPU_ALLOCATOR);
 				RecreateInfo->SwapChain = VK_NULL_HANDLE;
 			}
 			if (RecreateInfo->Surface)
@@ -491,7 +491,7 @@ void FVulkanViewport::CreateSwapchain(FVulkanCommandListContext& Context, FVulka
 #if VULKAN_ENABLE_DRAW_MARKERS
 		if (Device->GetSetDebugName())
 		{
-			VulkanRHI::SetDebugName(Device->GetSetDebugName(), Device->GetInstanceHandle(), RenderingBackBuffer->Image, "RenderingBackBuffer");
+			VulkanRHI::SetDebugName(Device->GetSetDebugName(), Device->GetHandle(), RenderingBackBuffer->Image, "RenderingBackBuffer");
 		}
 #endif
 	}
