VulkanRHI: Remove unused function.  Create dedicated flags for each extended dynamic state extension.  Start removing usage of deprecated VulkanSDK defines.

#rnx

diff --git a/Engine/Source/Runtime/VulkanRHI/Private/VulkanDevice.h b/Engine/Source/Runtime/VulkanRHI/Private/VulkanDevice.h
--- a/Engine/Source/Runtime/VulkanRHI/Private/VulkanDevice.h
+++ b/Engine/Source/Runtime/VulkanRHI/Private/VulkanDevice.h
@@ -100,7 +100,9 @@ struct FOptionalVulkanDeviceExtensions
 			uint64 HasKHRDynamicRendering : 1;
 			uint64 HasEXTSubgroupSizeControl : 1;
 			uint64 HasEXTPipelineCreationCacheControl : 1;
-			uint64 HasEXTExtendedDynamicStates : 1;  // (includes all three extensions)
+			uint64 HasEXTExtendedDynamicState1 : 1;
+			uint64 HasEXTExtendedDynamicState2 : 1;
+			uint64 HasEXTExtendedDynamicState3 : 1;
 
 			// Promoted to 1.4
 			uint64 HasEXTLoadStoreOpNone : 1;
diff --git a/Engine/Source/Runtime/VulkanRHI/Private/VulkanExtensions.cpp b/Engine/Source/Runtime/VulkanRHI/Private/VulkanExtensions.cpp
--- a/Engine/Source/Runtime/VulkanRHI/Private/VulkanExtensions.cpp
+++ b/Engine/Source/Runtime/VulkanRHI/Private/VulkanExtensions.cpp
@@ -1973,7 +1973,7 @@ public:
 	virtual void PostPhysicalDeviceFeatures(FOptionalVulkanDeviceExtensions& ExtensionFlags) override final
 	{
 		bRequirementsPassed = (ExtendedDynamicStateFeatures.extendedDynamicState == VK_TRUE);
-		ExtensionFlags.HasEXTExtendedDynamicStates = bRequirementsPassed;
+		ExtensionFlags.HasEXTExtendedDynamicState1 = bRequirementsPassed;
 	}
 
 	virtual void PreCreateDevice(VkDeviceCreateInfo& DeviceCreateInfo) override final
@@ -2006,15 +2006,8 @@ public:
 
 	virtual void PostPhysicalDeviceFeatures(FOptionalVulkanDeviceExtensions& ExtensionFlags) override final
 	{
-		if (ExtensionFlags.HasEXTExtendedDynamicStates)
-		{
-			bRequirementsPassed = (ExtendedDynamicState2Features.extendedDynamicState2 == VK_TRUE);
-			ExtensionFlags.HasEXTExtendedDynamicStates = bRequirementsPassed;
-		}
-		else
-		{
-			bRequirementsPassed = false;
-		}
+		bRequirementsPassed = (ExtendedDynamicState2Features.extendedDynamicState2 == VK_TRUE);
+		ExtensionFlags.HasEXTExtendedDynamicState2 = bRequirementsPassed;
 	}
 
 	virtual void PreCreateDevice(VkDeviceCreateInfo& DeviceCreateInfo) override final
@@ -2047,20 +2040,13 @@ public:
 
 	virtual void PostPhysicalDeviceFeatures(FOptionalVulkanDeviceExtensions& ExtensionFlags) override final
 	{
-		if (ExtensionFlags.HasEXTExtendedDynamicStates)
-		{
-			bRequirementsPassed = 
-				(ExtendedDynamicState3Features.extendedDynamicState3PolygonMode == VK_TRUE) &&
-				(ExtendedDynamicState3Features.extendedDynamicState3ColorBlendEnable == VK_TRUE) &&
-				(ExtendedDynamicState3Features.extendedDynamicState3ColorBlendEquation == VK_TRUE) &&
-				(ExtendedDynamicState3Features.extendedDynamicState3ColorWriteMask == VK_TRUE) &&
-				(ExtendedDynamicState3Features.extendedDynamicState3RasterizationSamples == VK_TRUE);
-			ExtensionFlags.HasEXTExtendedDynamicStates = bRequirementsPassed;
-		}
-		else
-		{
-			bRequirementsPassed = false;
-		}
+		bRequirementsPassed = 
+			(ExtendedDynamicState3Features.extendedDynamicState3PolygonMode == VK_TRUE) &&
+			(ExtendedDynamicState3Features.extendedDynamicState3ColorBlendEnable == VK_TRUE) &&
+			(ExtendedDynamicState3Features.extendedDynamicState3ColorBlendEquation == VK_TRUE) &&
+			(ExtendedDynamicState3Features.extendedDynamicState3ColorWriteMask == VK_TRUE) &&
+			(ExtendedDynamicState3Features.extendedDynamicState3RasterizationSamples == VK_TRUE);
+		ExtensionFlags.HasEXTExtendedDynamicState3 = bRequirementsPassed;
 	}
 
 	virtual void PreCreateDevice(VkDeviceCreateInfo& DeviceCreateInfo) override final
@@ -2211,9 +2197,9 @@ FVulkanDeviceExtensionArray FVulkanDeviceExtension::GetUESupportedDeviceExtensio
 	ADD_CUSTOM_EXTENSION(FVulkanKHRDynamicRenderingLocalReadExtension);	// must be kept AFTER DynamicRendering!
 	ADD_CUSTOM_EXTENSION(FVulkanKHRDynamicRenderingUnusedAttachmentsExtension);	// must be kept AFTER DynamicRendering!
 	ADD_CUSTOM_EXTENSION(FVulkanKHRDynamicRenderingUnusedAttachmentsExtension);	// must be kept AFTER DynamicRendering!
-	ADD_CUSTOM_EXTENSION(FVulkanEXTExtendedDynamicStateExtension);	// extended dynamic state extension must be kept in order (1,2,3)
-	ADD_CUSTOM_EXTENSION(FVulkanEXTExtendedDynamicState2Extension); // extended dynamic state extension must be kept in order (1,2,3)
-	ADD_CUSTOM_EXTENSION(FVulkanEXTExtendedDynamicState3Extension); // extended dynamic state extension must be kept in order (1,2,3)
+	ADD_CUSTOM_EXTENSION(FVulkanEXTExtendedDynamicStateExtension);
+	ADD_CUSTOM_EXTENSION(FVulkanEXTExtendedDynamicState2Extension);
+	ADD_CUSTOM_EXTENSION(FVulkanEXTExtendedDynamicState3Extension);
 	ADD_CUSTOM_EXTENSION(FVulkanEXTShaderObjectExtension);
 
 	// Needed for Raytracing
diff --git a/Engine/Source/Runtime/VulkanRHI/Private/VulkanPendingState.cpp b/Engine/Source/Runtime/VulkanRHI/Private/VulkanPendingState.cpp
--- a/Engine/Source/Runtime/VulkanRHI/Private/VulkanPendingState.cpp
+++ b/Engine/Source/Runtime/VulkanRHI/Private/VulkanPendingState.cpp
@@ -491,7 +491,7 @@ void FVulkanPendingGfxState::PrepareForDraw(FVulkanCommandListContext& Context)
 	}
 }
 
-void FVulkanPendingGfxState::InternalUpdateDynamicStates(FVulkanCommandBuffer& CommandBuffer)
+void FVulkanPendingGfxState::UpdateDynamicStates(FVulkanCommandBuffer& CommandBuffer)
 {
 	const bool bNeedsUpdateViewport = !CommandBuffer.bHasViewport || Viewports.Num() != CommandBuffer.CurrentViewports.Num() ||
 		(FMemory::Memcmp((const void*)CommandBuffer.CurrentViewports.GetData(), (const void*)Viewports.GetData(), Viewports.Num() * sizeof(VkViewport)) != 0);
diff --git a/Engine/Source/Runtime/VulkanRHI/Private/VulkanPendingState.h b/Engine/Source/Runtime/VulkanRHI/Private/VulkanPendingState.h
--- a/Engine/Source/Runtime/VulkanRHI/Private/VulkanPendingState.h
+++ b/Engine/Source/Runtime/VulkanRHI/Private/VulkanPendingState.h
@@ -199,7 +199,7 @@ public:
 		bScissorEnable = true;
 	}
 
-	inline void SetScissor(bool bInEnable, uint32 MinX, uint32 MinY, uint32 MaxX, uint32 MaxY)
+	void SetScissor(bool bInEnable, uint32 MinX, uint32 MinY, uint32 MaxX, uint32 MaxY)
 	{
 		if (bInEnable)
 		{
@@ -214,7 +214,7 @@ public:
 		bScissorEnable = bInEnable;
 	}
 
-	inline void SetScissorRect(uint32 MinX, uint32 MinY, uint32 Width, uint32 Height)
+	void SetScissorRect(uint32 MinX, uint32 MinY, uint32 Width, uint32 Height)
 	{
 		Scissors.SetNumZeroed(1);
 
@@ -224,59 +224,59 @@ public:
 		Scissors[0].extent.height = Height;
 	}
 
-	inline void SetStreamSource(uint32 StreamIndex, VkBuffer VertexBuffer, uint32 Offset)
+	void SetStreamSource(uint32 StreamIndex, VkBuffer VertexBuffer, uint32 Offset)
 	{
 		PendingStreams[StreamIndex].Stream = VertexBuffer;
 		PendingStreams[StreamIndex].BufferOffset = Offset;
 		bDirtyVertexStreams = true;
 	}
 
-	inline void Bind(VkCommandBuffer CmdBuffer)
+	void Bind(VkCommandBuffer CmdBuffer)
 	{
 		CurrentPipeline->Bind(CmdBuffer);
 	}
 
-	inline void SetTextureForStage(ShaderStage::EStage Stage, uint32 ParameterIndex, const FVulkanTexture* Texture, VkImageLayout Layout)
+	void SetTextureForStage(ShaderStage::EStage Stage, uint32 ParameterIndex, const FVulkanTexture* Texture, VkImageLayout Layout)
 	{
 		CurrentState->SetTexture(Stage, ParameterIndex, Texture, Layout);
 	}
 
-	inline void SetTextureForUBResource(uint8 DescriptorSet, uint32 BindingIndex, const FVulkanTexture* Texture, VkImageLayout Layout)
+	void SetTextureForUBResource(uint8 DescriptorSet, uint32 BindingIndex, const FVulkanTexture* Texture, VkImageLayout Layout)
 	{
 		CurrentState->SetTexture(DescriptorSet, BindingIndex, Texture, Layout);
 	}
 
 	template<bool bDynamic>
-	inline void SetUniformBuffer(uint8 DescriptorSet, uint32 BindingIndex, const FVulkanUniformBuffer* UniformBuffer)
+	void SetUniformBuffer(uint8 DescriptorSet, uint32 BindingIndex, const FVulkanUniformBuffer* UniformBuffer)
 	{
 		CurrentState->SetUniformBuffer<bDynamic>(DescriptorSet, BindingIndex, UniformBuffer);
 	}
 
 	void SetUAVForUBResource(uint8 DescriptorSet, uint32 BindingIndex, FVulkanUnorderedAccessView* UAV);
 
-	inline void SetUAVForStage(ShaderStage::EStage Stage, uint32 ParameterIndex, FVulkanUnorderedAccessView* UAV)
+	void SetUAVForStage(ShaderStage::EStage Stage, uint32 ParameterIndex, FVulkanUnorderedAccessView* UAV)
 	{
 		SetUAVForUBResource(Stage, ParameterIndex, UAV);
 	}
 
 	void SetSRVForUBResource(uint8 DescriptorSet, uint32 BindingIndex, FVulkanShaderResourceView* SRV);
 
-	inline void SetSRVForStage(ShaderStage::EStage Stage, uint32 ParameterIndex, FVulkanShaderResourceView* SRV)
+	void SetSRVForStage(ShaderStage::EStage Stage, uint32 ParameterIndex, FVulkanShaderResourceView* SRV)
 	{
 		SetSRVForUBResource(Stage, ParameterIndex, SRV);
 	}
 
-	inline void SetSamplerStateForStage(ShaderStage::EStage Stage, uint32 ParameterIndex, FVulkanSamplerState* Sampler)
+	void SetSamplerStateForStage(ShaderStage::EStage Stage, uint32 ParameterIndex, FVulkanSamplerState* Sampler)
 	{
 		CurrentState->SetSamplerState(Stage, ParameterIndex, Sampler);
 	}
 
-	inline void SetSamplerStateForUBResource(uint32 DescriptorSet, uint32 BindingIndex, FVulkanSamplerState* Sampler)
+	void SetSamplerStateForUBResource(uint32 DescriptorSet, uint32 BindingIndex, FVulkanSamplerState* Sampler)
 	{
 		CurrentState->SetSamplerState(DescriptorSet, BindingIndex, Sampler);
 	}
 
-	inline void SetPackedGlobalShaderParameter(ShaderStage::EStage Stage, uint32 BufferIndex, uint32 Offset, uint32 NumBytes, const void* NewValue)
+	void SetPackedGlobalShaderParameter(ShaderStage::EStage Stage, uint32 BufferIndex, uint32 Offset, uint32 NumBytes, const void* NewValue)
 	{
 		CurrentState->SetPackedGlobalShaderParameter(Stage, BufferIndex, Offset, NumBytes, NewValue);
 	}
@@ -285,12 +285,9 @@ public:
 
 	bool SetGfxPipeline(FVulkanGraphicsPipelineState* InGfxPipeline, bool bForceReset);
 
-	inline void UpdateDynamicStates(FVulkanCommandBuffer& CommandBuffer)
-	{
-		InternalUpdateDynamicStates(CommandBuffer);
-	}
+	void UpdateDynamicStates(FVulkanCommandBuffer& CommandBuffer);
 
-	inline void SetStencilRef(uint32 InStencilRef)
+	void SetStencilRef(uint32 InStencilRef)
 	{
 		if (InStencilRef != StencilRef)
 		{
@@ -303,10 +300,6 @@ public:
 		PipelineStates.Remove(Pipeline);
 	}
 
-	inline void MarkNeedsDynamicStates()
-	{
-	}
-
 protected:
 	TArray<VkViewport, TInlineAllocator<2>> Viewports;
 	TArray<VkRect2D, TInlineAllocator<2>> Scissors;
@@ -336,7 +329,6 @@ protected:
 	FVertexStream PendingStreams[MaxVertexElementCount];
 	bool bDirtyVertexStreams;
 
-	void InternalUpdateDynamicStates(FVulkanCommandBuffer& CommandBuffer);
 	void UpdateInputAttachments(FVulkanFramebuffer* Framebuffer);
 
 	FVulkanDevice& Device;
diff --git a/Engine/Source/Runtime/VulkanRHI/Private/VulkanPipeline.cpp b/Engine/Source/Runtime/VulkanRHI/Private/VulkanPipeline.cpp
--- a/Engine/Source/Runtime/VulkanRHI/Private/VulkanPipeline.cpp
+++ b/Engine/Source/Runtime/VulkanRHI/Private/VulkanPipeline.cpp
@@ -1385,13 +1385,13 @@ bool FVulkanPipelineStateCacheManager::CreateGfxPipelineFromEntry(FVulkanGraphic
 
 	VkPipelineDynamicStateCreateInfo DynamicState;
 	ZeroVulkanStruct(DynamicState, VK_STRUCTURE_TYPE_PIPELINE_DYNAMIC_STATE_CREATE_INFO);
-	VkDynamicState DynamicStatesEnabled[VK_DYNAMIC_STATE_RANGE_SIZE];
-	DynamicState.pDynamicStates = DynamicStatesEnabled;
-	FMemory::Memzero(DynamicStatesEnabled);
-	DynamicStatesEnabled[DynamicState.dynamicStateCount++] = VK_DYNAMIC_STATE_VIEWPORT;
-	DynamicStatesEnabled[DynamicState.dynamicStateCount++] = VK_DYNAMIC_STATE_SCISSOR;
-	DynamicStatesEnabled[DynamicState.dynamicStateCount++] = VK_DYNAMIC_STATE_STENCIL_REFERENCE;
-	DynamicStatesEnabled[DynamicState.dynamicStateCount++] = VK_DYNAMIC_STATE_DEPTH_BOUNDS;
+	TArray<VkDynamicState, TInlineAllocator<4>> DynamicStatesEnabled;
+	DynamicStatesEnabled.Add(VK_DYNAMIC_STATE_VIEWPORT);
+	DynamicStatesEnabled.Add(VK_DYNAMIC_STATE_SCISSOR);
+	DynamicStatesEnabled.Add(VK_DYNAMIC_STATE_STENCIL_REFERENCE);
+	DynamicStatesEnabled.Add(VK_DYNAMIC_STATE_DEPTH_BOUNDS);
+	DynamicState.pDynamicStates = DynamicStatesEnabled.GetData();
+	DynamicState.dynamicStateCount = DynamicStatesEnabled.Num();
 
 	PipelineInfo.pDynamicState = &DynamicState;
 
diff --git a/Engine/Source/Runtime/VulkanRHI/Private/VulkanPipelineState.cpp b/Engine/Source/Runtime/VulkanRHI/Private/VulkanPipelineState.cpp
--- a/Engine/Source/Runtime/VulkanRHI/Private/VulkanPipelineState.cpp
+++ b/Engine/Source/Runtime/VulkanRHI/Private/VulkanPipelineState.cpp
@@ -584,7 +584,6 @@ void FVulkanCommandListContext::RHISetGraphicsPipelineState(FRHIGraphicsPipeline
 #endif
 		PendingGfxState->Bind(CommandBuffer.GetHandle());
 		CommandBuffer.bHasPipeline = true;
-		PendingGfxState->MarkNeedsDynamicStates();
 	}
 
 	PendingGfxState->SetStencilRef(StencilRef);
diff --git a/Engine/Source/Runtime/VulkanRHI/Public/VulkanThirdParty.h b/Engine/Source/Runtime/VulkanRHI/Public/VulkanThirdParty.h
--- a/Engine/Source/Runtime/VulkanRHI/Public/VulkanThirdParty.h
+++ b/Engine/Source/Runtime/VulkanRHI/Public/VulkanThirdParty.h
@@ -47,9 +47,6 @@ THIRD_PARTY_INCLUDES_END
 	#define VK_DESCRIPTOR_TYPE_END_RANGE (VK_DESCRIPTOR_TYPE_INPUT_ATTACHMENT)
 	#define VK_DESCRIPTOR_TYPE_RANGE_SIZE (VK_DESCRIPTOR_TYPE_INPUT_ATTACHMENT - VK_DESCRIPTOR_TYPE_SAMPLER + 1)
 	#define VK_IMAGE_VIEW_TYPE_RANGE_SIZE (VK_IMAGE_VIEW_TYPE_CUBE_ARRAY - VK_IMAGE_VIEW_TYPE_1D + 1)
-	#define VK_DYNAMIC_STATE_BEGIN_RANGE (VK_DYNAMIC_STATE_VIEWPORT)
-	#define VK_DYNAMIC_STATE_END_RANGE (VK_DYNAMIC_STATE_STENCIL_REFERENCE)
-	#define VK_DYNAMIC_STATE_RANGE_SIZE (VK_DYNAMIC_STATE_STENCIL_REFERENCE - VK_DYNAMIC_STATE_VIEWPORT + 1)
 	#define VK_FORMAT_RANGE_SIZE (VK_FORMAT_ASTC_12x12_SRGB_BLOCK - VK_FORMAT_UNDEFINED + 1)
 #endif
 
