VulkanRHI: Remove unused ControlPoints from FGfxPipelineDesc

#rnx

diff --git a/Engine/Source/Runtime/VulkanRHI/Private/VulkanPipeline.cpp b/Engine/Source/Runtime/VulkanRHI/Private/VulkanPipeline.cpp
--- a/Engine/Source/Runtime/VulkanRHI/Private/VulkanPipeline.cpp
+++ b/Engine/Source/Runtime/VulkanRHI/Private/VulkanPipeline.cpp
@@ -1098,7 +1098,6 @@ FArchive& operator << (FArchive& Ar, FGfxPipelineDesc& Entry)
 	// Modify VERSION if serialization changes
 	Ar << Entry.VertexInputKey;
 	Ar << Entry.RasterizationSamples;
-	Ar << Entry.ControlPoints;
 	Ar << Entry.Topology;
 
 	Ar << Entry.ColorAttachmentStates;
@@ -1799,9 +1798,8 @@ FVulkanLayout* FVulkanPipelineStateCacheManager::FindOrAddLayout(const FVulkanDe
 	return Layout;
 }
 
-static inline VkPrimitiveTopology UEToVulkanTopologyType(const FVulkanDevice* InDevice, EPrimitiveType PrimitiveType, uint16& OutControlPoints)
+static inline VkPrimitiveTopology UEToVulkanTopologyType(EPrimitiveType PrimitiveType)
 {
-	OutControlPoints = 0;
 	switch (PrimitiveType)
 	{
 	case PT_PointList:
@@ -1890,7 +1888,7 @@ void FVulkanPipelineStateCacheManager::CreateGfxEntry(const FGraphicsPipelineSta
 	OutGfxEntry->UseAlphaToCoverage = PSOInitializer.NumSamples > 1 && BlendState->Initializer.bUseAlphaToCoverage ? 1 : 0;
 
 	OutGfxEntry->RasterizationSamples = PSOInitializer.NumSamples;
-	OutGfxEntry->Topology = (uint32)UEToVulkanTopologyType(Device, PSOInitializer.PrimitiveType, OutGfxEntry->ControlPoints);
+	OutGfxEntry->Topology = (uint32)UEToVulkanTopologyType(PSOInitializer.PrimitiveType);
 	uint32 NumRenderTargets = PSOInitializer.ComputeNumValidRenderTargets();
 	
 	if (PSOInitializer.SubpassHint == ESubpassHint::DeferredShadingSubpass && PSOInitializer.SubpassIndex >= 2)
diff --git a/Engine/Source/Runtime/VulkanRHI/Private/VulkanPipeline.h b/Engine/Source/Runtime/VulkanRHI/Private/VulkanPipeline.h
--- a/Engine/Source/Runtime/VulkanRHI/Private/VulkanPipeline.h
+++ b/Engine/Source/Runtime/VulkanRHI/Private/VulkanPipeline.h
@@ -129,7 +129,6 @@ struct FGfxPipelineDesc
 
 	uint32 VertexInputKey;
 	uint16 RasterizationSamples;
-	uint16 ControlPoints;
 	uint32 Topology;
 	struct FBlendAttachment
 	{
