Vulkan Bindless changes
- Use the RHICore descriptor allocators.
- Rename ReserveDescriptor to AllocateDescriptor.
- Adding a few AllocateDescriptor shortcuts for initialization.
- Some other minor changes/reorganization.

diff --git a/Engine/Source/Runtime/VulkanRHI/Private/VulkanBindlessDescriptorManager.cpp b/Engine/Source/Runtime/VulkanRHI/Private/VulkanBindlessDescriptorManager.cpp
--- a/Engine/Source/Runtime/VulkanRHI/Private/VulkanBindlessDescriptorManager.cpp
+++ b/Engine/Source/Runtime/VulkanRHI/Private/VulkanBindlessDescriptorManager.cpp
@@ -80,28 +80,36 @@ static FAutoConsoleVariableRef CVarVulkanBindlessBlockSize(
 
 DECLARE_STATS_GROUP(TEXT("Vulkan Bindless"), STATGROUP_VulkanBindless, STATCAT_Advanced);
 
-DECLARE_DWORD_ACCUMULATOR_STAT_EXTERN(TEXT("Num Peak Descriptor Count"), STAT_VulkanBindlessPeakDescriptorCount, STATGROUP_VulkanBindless, );
-DECLARE_DWORD_ACCUMULATOR_STAT_EXTERN(TEXT("Peak Samplers"), STAT_VulkanBindlessPeakSampler, STATGROUP_VulkanBindless, );
-DECLARE_DWORD_ACCUMULATOR_STAT_EXTERN(TEXT("Peak Sampled Images"), STAT_VulkanBindlessPeakSampledImage, STATGROUP_VulkanBindless, );
-DECLARE_DWORD_ACCUMULATOR_STAT_EXTERN(TEXT("Peak Storage Images"), STAT_VulkanBindlessPeakStorageImage, STATGROUP_VulkanBindless, );
-DECLARE_DWORD_ACCUMULATOR_STAT_EXTERN(TEXT("Peak Uniform Buffers"), STAT_VulkanBindlessPeakUniformBuffer, STATGROUP_VulkanBindless, );
-DECLARE_DWORD_ACCUMULATOR_STAT_EXTERN(TEXT("Peak Storage Buffers"), STAT_VulkanBindlessPeakStorageBuffer, STATGROUP_VulkanBindless, );
-DECLARE_DWORD_ACCUMULATOR_STAT_EXTERN(TEXT("Peak Uniform Texel Buffers"), STAT_VulkanBindlessPeakUniformTexelBuffer, STATGROUP_VulkanBindless, );
-DECLARE_DWORD_ACCUMULATOR_STAT_EXTERN(TEXT("Peak Storage Texel Buffers"), STAT_VulkanBindlessPeakStorageTexelBuffer, STATGROUP_VulkanBindless, );
-DECLARE_DWORD_ACCUMULATOR_STAT_EXTERN(TEXT("Peak Acceleration Structures"), STAT_VulkanBindlessPeakAccelerationStructure, STATGROUP_VulkanBindless, );
-DECLARE_DWORD_COUNTER_STAT_EXTERN(TEXT("Write Per Frame"), STAT_VulkanBindlessWritePerFrame, STATGROUP_VulkanBindless, );
-
-DEFINE_STAT(STAT_VulkanBindlessPeakDescriptorCount);
-DEFINE_STAT(STAT_VulkanBindlessPeakSampler);
-DEFINE_STAT(STAT_VulkanBindlessPeakSampledImage);
-DEFINE_STAT(STAT_VulkanBindlessPeakStorageImage);
-DEFINE_STAT(STAT_VulkanBindlessPeakUniformBuffer);
-DEFINE_STAT(STAT_VulkanBindlessPeakStorageBuffer);
-DEFINE_STAT(STAT_VulkanBindlessPeakUniformTexelBuffer);
-DEFINE_STAT(STAT_VulkanBindlessPeakStorageTexelBuffer);
-DEFINE_STAT(STAT_VulkanBindlessPeakAccelerationStructure);
-DEFINE_STAT(STAT_VulkanBindlessWritePerFrame);
+DECLARE_DWORD_ACCUMULATOR_STAT(TEXT("Descriptors Allocated"), STAT_VulkanBindlessDescriptorsAllocated, STATGROUP_VulkanBindless);
 
+DECLARE_DWORD_ACCUMULATOR_STAT(TEXT("Sampler Descriptors Allocated"),               STAT_VulkanBindlessSamplerDescriptorsAllocated,               STATGROUP_VulkanBindless);
+DECLARE_DWORD_ACCUMULATOR_STAT(TEXT("SampledImage Descriptors Allocated"),          STAT_VulkanBindlessSampledImageDescriptorsAllocated,          STATGROUP_VulkanBindless);
+DECLARE_DWORD_ACCUMULATOR_STAT(TEXT("StorageImage Descriptors Allocated"),          STAT_VulkanBindlessStorageImageDescriptorsAllocated,          STATGROUP_VulkanBindless);
+DECLARE_DWORD_ACCUMULATOR_STAT(TEXT("UniformTexelBuffer Descriptors Allocated"),    STAT_VulkanBindlessUniformTexelBufferDescriptorsAllocated,    STATGROUP_VulkanBindless);
+DECLARE_DWORD_ACCUMULATOR_STAT(TEXT("StorageTexelBuffer Descriptors Allocated"),    STAT_VulkanBindlessStorageTexelBufferDescriptorsAllocated,    STATGROUP_VulkanBindless);
+DECLARE_DWORD_ACCUMULATOR_STAT(TEXT("UniformBuffer Descriptors Allocated"),         STAT_VulkanBindlessUniformBufferDescriptorsAllocated,         STATGROUP_VulkanBindless);
+DECLARE_DWORD_ACCUMULATOR_STAT(TEXT("StorageBuffer Descriptors Allocated"),         STAT_VulkanBindlessStorageBufferDescriptorsAllocated,         STATGROUP_VulkanBindless);
+DECLARE_DWORD_ACCUMULATOR_STAT(TEXT("AccelerationStructure Descriptors Allocated"), STAT_VulkanBindlessAccelerationStructureDescriptorsAllocated, STATGROUP_VulkanBindless);
+
+DECLARE_DWORD_COUNTER_STAT(TEXT("Write Per Frame"), STAT_VulkanBindlessWritePerFrame, STATGROUP_VulkanBindless);
+
+static TStatId GetStatForDescriptorType(VkDescriptorType DescriptorType)
+{
+	switch (DescriptorType)
+	{
+	case VK_DESCRIPTOR_TYPE_SAMPLER:                    return GET_STATID(STAT_VulkanBindlessSamplerDescriptorsAllocated);
+	case VK_DESCRIPTOR_TYPE_SAMPLED_IMAGE:              return GET_STATID(STAT_VulkanBindlessSampledImageDescriptorsAllocated);
+	case VK_DESCRIPTOR_TYPE_STORAGE_IMAGE:              return GET_STATID(STAT_VulkanBindlessStorageImageDescriptorsAllocated);
+	case VK_DESCRIPTOR_TYPE_UNIFORM_TEXEL_BUFFER:       return GET_STATID(STAT_VulkanBindlessUniformTexelBufferDescriptorsAllocated);
+	case VK_DESCRIPTOR_TYPE_STORAGE_TEXEL_BUFFER:       return GET_STATID(STAT_VulkanBindlessStorageTexelBufferDescriptorsAllocated);
+	case VK_DESCRIPTOR_TYPE_STORAGE_BUFFER:             return GET_STATID(STAT_VulkanBindlessUniformBufferDescriptorsAllocated);
+	case VK_DESCRIPTOR_TYPE_UNIFORM_BUFFER:             return GET_STATID(STAT_VulkanBindlessStorageBufferDescriptorsAllocated);
+	case VK_DESCRIPTOR_TYPE_ACCELERATION_STRUCTURE_KHR: return GET_STATID(STAT_VulkanBindlessAccelerationStructureDescriptorsAllocated);
+	default: checkNoEntry();
+	}
+
+	return {};
+}
 
 static constexpr uint8 GetIndexForDescriptorType(VkDescriptorType DescriptorType)
 {
@@ -158,7 +166,6 @@ static constexpr ERHIDescriptorType GetUEDescriptorTypeFromVkDescriptorType(VkDe
 	return ERHIDescriptorType::Invalid;
 }
 
-
 extern TAutoConsoleVariable<int32> GCVarRobustBufferAccess;
 static inline uint32 GetDescriptorTypeSize(FVulkanDevice& Device, VkDescriptorType DescriptorType)
 {
@@ -222,7 +229,19 @@ static inline VkMemoryPropertyFlags GetDescriptorBufferMemoryType(FVulkanDevice&
 	return VK_MEMORY_PROPERTY_HOST_VISIBLE_BIT | VK_MEMORY_PROPERTY_HOST_COHERENT_BIT;
 }
 
+inline FVulkanBindlessDescriptorManager::BindlessSetState& FVulkanBindlessDescriptorManager::GetBindlessState(ERHIDescriptorType DescriptorType)
+{
+	const uint8 SetIndex = GetIndexForDescriptorType(DescriptorType);
+	check(SetIndex < VulkanBindless::NumBindlessSets);
+	return BindlessSetStates[SetIndex];
+}
 
+inline FVulkanBindlessDescriptorManager::BindlessSetState& FVulkanBindlessDescriptorManager::GetBindlessState(VkDescriptorType DescriptorType)
+{
+	const uint8 SetIndex = GetIndexForDescriptorType(DescriptorType);
+	check(SetIndex < VulkanBindless::NumBindlessSets);
+	return BindlessSetStates[SetIndex];
+}
 
 // Check all the requirements to be running in Bindless using Descriptor Buffers
 bool FVulkanBindlessDescriptorManager::VerifySupport(FVulkanDevice& Device)
@@ -370,19 +389,18 @@ void FVulkanBindlessDescriptorManager::Init()
 	}
 
 	{
-		auto InitBindlessSetState = [&BindlessSetStates = BindlessSetStates, &Device = Device](VkDescriptorType DescriptorType) {
-
-			const uint8 StateIndex = GetIndexForDescriptorType(DescriptorType);
-			BindlessSetState& OutState = BindlessSetStates[StateIndex];
+		auto InitBindlessSetState = [this](VkDescriptorType DescriptorType)
+		{
+			BindlessSetState& OutState = GetBindlessState(DescriptorType);
 
 			OutState.DescriptorTypes.Add(DescriptorType);
 
 			OutState.DescriptorSize = FMath::Max(OutState.DescriptorSize, GetDescriptorTypeSize(Device, DescriptorType));
 			checkf((OutState.DescriptorSize > 0), TEXT("Descriptor Type [%s] returned an invalid descriptor size!"), VK_TYPE_TO_STRING(VkDescriptorType, DescriptorType));
 
-			OutState.MaxDescriptorCount += GetInitialDescriptorCount(DescriptorType);
-			checkf((OutState.MaxDescriptorCount > 0), TEXT("Descriptor Type [%s] returned an invalid descriptor count!"), VK_TYPE_TO_STRING(VkDescriptorType, DescriptorType));
-			};
+			OutState.DescriptorCapacity += GetInitialDescriptorCount(DescriptorType);
+			checkf((OutState.DescriptorCapacity > 0), TEXT("Descriptor Type [%s] returned an invalid descriptor count!"), VK_TYPE_TO_STRING(VkDescriptorType, DescriptorType));
+		};
 
 		// Go through all the supported descriptor types in bindless
 		InitBindlessSetState(VK_DESCRIPTOR_TYPE_SAMPLER);
@@ -400,16 +418,15 @@ void FVulkanBindlessDescriptorManager::Init()
 		}
 		else
 		{
-			const uint8 StateIndex = GetIndexForDescriptorType(VK_DESCRIPTOR_TYPE_ACCELERATION_STRUCTURE_KHR);
-			BindlessSetState& State = BindlessSetStates[StateIndex];
+			BindlessSetState& State = GetBindlessState(VK_DESCRIPTOR_TYPE_ACCELERATION_STRUCTURE_KHR);
 			State.DescriptorTypes.Add(VK_DESCRIPTOR_TYPE_STORAGE_BUFFER);
 			State.DescriptorSize = GetDescriptorTypeSize(Device, VK_DESCRIPTOR_TYPE_STORAGE_BUFFER);
-			State.MaxDescriptorCount = 16;
+			State.DescriptorCapacity = 16;
 		}
 
 		// Fill the DescriptorSetLayout for a BindlessSetState
-		auto CreateDescriptorSetLayout = [&](const BindlessSetState& State) {
-
+		auto CreateDescriptorSetLayout = [&](const BindlessSetState& State)
+		{
 			if (State.DescriptorTypes.Num() == 0)
 			{
 				return EmptyDescriptorSetLayout;
@@ -418,7 +435,7 @@ void FVulkanBindlessDescriptorManager::Init()
 			{
 				VkDescriptorSetLayoutBinding DescriptorSetLayoutBinding;
 				DescriptorSetLayoutBinding.binding = 0;
-				DescriptorSetLayoutBinding.descriptorCount = State.MaxDescriptorCount;  // todo-jn: resizable
+				DescriptorSetLayoutBinding.descriptorCount = State.DescriptorCapacity;  // todo-jn: resizable
 				DescriptorSetLayoutBinding.stageFlags = VK_SHADER_STAGE_ALL;
 				DescriptorSetLayoutBinding.pImmutableSamplers = nullptr; // todo-jn: ImmutableSamplers
 
@@ -463,17 +480,39 @@ void FVulkanBindlessDescriptorManager::Init()
 				VERIFYVULKANRESULT(VulkanRHI::vkCreateDescriptorSetLayout(DeviceHandle, &DescriptorSetLayoutCreateInfo, VULKAN_CPU_ALLOCATOR, &DescriptorSetLayout));
 				return DescriptorSetLayout;
 			}
-			};
+		};
 
 		// Create the descriptor buffer for a BindlessSetState
-		auto CreateDescriptorBuffer = [&](BindlessSetState& InOutState, VkDescriptorBufferBindingInfoEXT& OutBufferBindingInfo) {
-
+		auto CreateDescriptorBuffer = [&](BindlessSetState& InOutState, VkDescriptorBufferBindingInfoEXT& OutBufferBindingInfo)
+		{
 			// Skip unsupported descriptors
 			if (InOutState.DescriptorTypes.Num() == 0)
 			{
 				return 0u;
 			}
 
+			ERHIDescriptorTypeMask TypeMask{};
+			TArray<TStatId, TInlineAllocator<3>> AllocatorStats;
+
+			AllocatorStats.Emplace(GET_STATID(STAT_VulkanBindlessDescriptorsAllocated));
+
+			for (VkDescriptorType DescriptorType : InOutState.DescriptorTypes)
+			{
+				EnumAddFlags(TypeMask, RHIDescriptorTypeMaskFromType(GetUEDescriptorTypeFromVkDescriptorType(DescriptorType)));
+				AllocatorStats.Emplace(GetStatForDescriptorType(DescriptorType));
+
+				if (DescriptorType == VK_DESCRIPTOR_TYPE_SAMPLER)
+				{
+					AllocatorStats.Emplace(GET_STATID(STAT_BindlessSamplerDescriptorsAllocated));
+				}
+				else
+				{
+					AllocatorStats.Emplace(GET_STATID(STAT_BindlessResourceDescriptorsAllocated));
+				}
+			}
+
+			InOutState.Allocator = new FRHIHeapDescriptorAllocator(TypeMask, InOutState.DescriptorCapacity, TConstArrayView<TStatId>(AllocatorStats));
+
 			const bool IsSamplerSet = (InOutState.DescriptorTypes[0] == VK_DESCRIPTOR_TYPE_SAMPLER);
 			const VkBufferUsageFlags BufferUsageFlags =
 				VK_BUFFER_USAGE_SHADER_DEVICE_ADDRESS_BIT |
@@ -481,13 +520,13 @@ void FVulkanBindlessDescriptorManager::Init()
 				VK_BUFFER_USAGE_TRANSFER_DST_BIT |
 				(IsSamplerSet ? VK_BUFFER_USAGE_SAMPLER_DESCRIPTOR_BUFFER_BIT_EXT : VK_BUFFER_USAGE_RESOURCE_DESCRIPTOR_BUFFER_BIT_EXT);
 
-			const uint32 DescriptorBufferSize = InOutState.DescriptorSize * InOutState.MaxDescriptorCount;
-			InOutState.DebugDescriptors.SetNumZeroed(DescriptorBufferSize);
+			const uint32 DescriptorBufferSize = InOutState.DescriptorSize * InOutState.DescriptorCapacity;
+			InOutState.CPUDescriptorMemory = reinterpret_cast<uint8*>(FMemory::MallocZeroed(DescriptorBufferSize));
 
 			VkDeviceSize LayoutSizeInBytes = 0;
 			VulkanRHI::vkGetDescriptorSetLayoutSizeEXT(DeviceHandle, InOutState.DescriptorSetLayout, &LayoutSizeInBytes);
 			// Double check that the layout follows the rules for a single binding with an array of descriptors that are tightly packed
-			check(LayoutSizeInBytes == (InOutState.MaxDescriptorCount * InOutState.DescriptorSize));
+			check(LayoutSizeInBytes == DescriptorBufferSize);
 
 			if (IsSamplerSet)
 			{
@@ -539,7 +578,7 @@ void FVulkanBindlessDescriptorManager::Init()
 			}
 
 			return IsSamplerSet ? 0u : DescriptorBufferSize;
-			};
+		};
 
 		// Fill in one state for each descriptor type
 		uint32 TotalResourceDescriptorBufferSize = 0;
@@ -681,59 +720,84 @@ void FVulkanBindlessDescriptorManager::RegisterUniformBuffers(FVulkanCommandList
 	VulkanRHI::vkCmdSetDescriptorBufferOffsetsEXT(CommandBufferHandle, BindPoint, BindlessPipelineLayout, 0u, VulkanBindless::MaxNumSets, BufferIndices, BufferOffsets);
 }
 
-void FVulkanBindlessDescriptorManager::UpdateStatsForHandle(VkDescriptorType DescriptorType)
+FRHIDescriptorHandle FVulkanBindlessDescriptorManager::AllocateDescriptor(VkDescriptorType DescriptorType)
 {
-	const uint8 SetIndex = GetIndexForDescriptorType(DescriptorType);
-	const BindlessSetState& State = BindlessSetStates[SetIndex];
+	if (IsSupported())
+	{
+		ERHIDescriptorType UEType = GetUEDescriptorTypeFromVkDescriptorType(DescriptorType);
+		return GetBindlessState(DescriptorType).Allocator->Allocate(UEType);
+	}
 
-	switch (DescriptorType)
+	return FRHIDescriptorHandle();
+}
+
+FRHIDescriptorHandle FVulkanBindlessDescriptorManager::AllocateDescriptor(VkDescriptorType DescriptorType, VkDeviceAddress BufferAddress, VkDeviceSize BufferSize)
+{
+	if (IsSupported())
 	{
-	case VK_DESCRIPTOR_TYPE_SAMPLER:                   SET_DWORD_STAT(STAT_VulkanBindlessPeakSampler, State.PeakDescriptorCount); break;
-	case VK_DESCRIPTOR_TYPE_SAMPLED_IMAGE:             SET_DWORD_STAT(STAT_VulkanBindlessPeakSampledImage, State.PeakDescriptorCount); break;
-	case VK_DESCRIPTOR_TYPE_STORAGE_IMAGE:             SET_DWORD_STAT(STAT_VulkanBindlessPeakStorageImage, State.PeakDescriptorCount); break;
-	case VK_DESCRIPTOR_TYPE_UNIFORM_TEXEL_BUFFER:      SET_DWORD_STAT(STAT_VulkanBindlessPeakUniformTexelBuffer, State.PeakDescriptorCount); break;
-	case VK_DESCRIPTOR_TYPE_STORAGE_TEXEL_BUFFER:      SET_DWORD_STAT(STAT_VulkanBindlessPeakStorageTexelBuffer, State.PeakDescriptorCount); break;
-	case VK_DESCRIPTOR_TYPE_STORAGE_BUFFER:            SET_DWORD_STAT(STAT_VulkanBindlessPeakStorageBuffer, State.PeakDescriptorCount); break;
-	case VK_DESCRIPTOR_TYPE_UNIFORM_BUFFER:            SET_DWORD_STAT(STAT_VulkanBindlessPeakUniformBuffer, State.PeakDescriptorCount); break;
-	case VK_DESCRIPTOR_TYPE_ACCELERATION_STRUCTURE_KHR:SET_DWORD_STAT(STAT_VulkanBindlessPeakAccelerationStructure, State.PeakDescriptorCount); break;
-	default:
-		checkf(false, TEXT("Unknown descriptor type for stat update: %s"), VK_TYPE_TO_STRING(VkDescriptorType, DescriptorType));
+		FRHIDescriptorHandle DescriptorHandle = AllocateDescriptor(DescriptorType);
+		UpdateBuffer(DescriptorHandle, DescriptorType, BufferAddress, BufferSize, true);
+		return DescriptorHandle;
 	}
+
+	return FRHIDescriptorHandle();
 }
 
-FRHIDescriptorHandle FVulkanBindlessDescriptorManager::ReserveDescriptor(VkDescriptorType DescriptorType)
+FRHIDescriptorHandle FVulkanBindlessDescriptorManager::AllocateDescriptor(VkDescriptorType DescriptorType, FVulkanBuffer* Buffer, uint32 ExtraOffset)
 {
-	if (bIsSupported)
+	if (IsSupported())
 	{
-		const uint8 SetIndex = GetIndexForDescriptorType(DescriptorType);
-		const ERHIDescriptorType RHIType = GetUEDescriptorTypeFromVkDescriptorType(DescriptorType);
+		FRHIDescriptorHandle DescriptorHandle = AllocateDescriptor(DescriptorType);
+		UpdateBuffer(DescriptorHandle, DescriptorType, Buffer->GetHandle(), Buffer->GetOffset() + ExtraOffset, Buffer->GetCurrentSize() - ExtraOffset, true);
+		return DescriptorHandle;
+	}
 
-		BindlessSetState& State = BindlessSetStates[SetIndex];
-		const uint32 ResourceIndex = GetFreeResourceIndex(State);
-		return FRHIDescriptorHandle(RHIType, ResourceIndex);
+	return FRHIDescriptorHandle();
+}
+
+FRHIDescriptorHandle FVulkanBindlessDescriptorManager::AllocateDescriptor(VkSampler VulkanSampler)
+{
+	if (IsSupported())
+	{
+		FRHIDescriptorHandle DescriptorHandle = AllocateDescriptor(VK_DESCRIPTOR_TYPE_SAMPLER);
+		UpdateSampler(DescriptorHandle, VulkanSampler);
+		return DescriptorHandle;
 	}
+
 	return FRHIDescriptorHandle();
 }
 
+void FVulkanBindlessDescriptorManager::FreeDescriptor(FRHIDescriptorHandle DescriptorHandle)
+{
+	if (DescriptorHandle.IsValid())
+	{
+		BindlessSetState& State = GetBindlessState(DescriptorHandle.GetType());
+		State.Allocator->Free(DescriptorHandle);
+
+		const uint32 ByteOffset = DescriptorHandle.GetIndex() * State.DescriptorSize;
+		FMemory::Memzero(&State.CPUDescriptorMemory[ByteOffset], State.DescriptorSize); // easier for debugging for now
+	}
+}
+
 void FVulkanBindlessDescriptorManager::UpdateDescriptor(FRHIDescriptorHandle DescriptorHandle, VkDescriptorType DescriptorType, VkDescriptorDataEXT DescriptorData, bool bImmediateUpdate)
 {
 	checkf(DescriptorHandle.IsValid(), TEXT("Attemping to update invalid descriptor handle!"));
 
-	const uint8 SetIndex = GetIndexForDescriptorType(DescriptorHandle.GetType());
-
-	BindlessSetState& State = BindlessSetStates[SetIndex];
+	BindlessSetState& State = GetBindlessState(DescriptorHandle.GetType());
 	const uint32 ByteOffset = DescriptorHandle.GetIndex() * State.DescriptorSize;
 	checkSlow(State.DescriptorTypes.Contains(DescriptorType));
 
+	uint8* const CPUDescriptor = &State.CPUDescriptorMemory[ByteOffset];
+
 	VkDescriptorGetInfoEXT Info;
 	ZeroVulkanStruct(Info, VK_STRUCTURE_TYPE_DESCRIPTOR_GET_INFO_EXT);
 	Info.type = DescriptorType;
 	Info.data = DescriptorData;
-	VulkanRHI::vkGetDescriptorEXT(Device.GetHandle(), &Info, State.DescriptorSize, &State.DebugDescriptors[ByteOffset]);
+	VulkanRHI::vkGetDescriptorEXT(Device.GetHandle(), &Info, State.DescriptorSize, CPUDescriptor);
 
 	if (bImmediateUpdate)
 	{
-		FMemory::Memcpy(&State.MappedPointer[ByteOffset], &State.DebugDescriptors[ByteOffset], State.DescriptorSize);
+		FMemory::Memcpy(&State.MappedPointer[ByteOffset], CPUDescriptor, State.DescriptorSize);
 	}
 	else
 	{
@@ -743,12 +807,12 @@ void FVulkanBindlessDescriptorManager::UpdateDescriptor(FRHIDescriptorHandle Des
 		// :todo-jn:  Hack to avoid barriers/copies in renderpasses
 		if (CmdBuffer->IsInsideRenderPass())
 		{
-			FMemory::Memcpy(&State.MappedPointer[ByteOffset], &State.DebugDescriptors[ByteOffset], State.DescriptorSize);
+			FMemory::Memcpy(&State.MappedPointer[ByteOffset], CPUDescriptor, State.DescriptorSize);
 		}
 		else
 		{
 			VulkanRHI::FStagingBuffer* StagingBuffer = Device.GetStagingManager().AcquireBuffer(State.DescriptorSize);
-			FMemory::Memcpy(StagingBuffer->GetMappedPointer(), &State.DebugDescriptors[ByteOffset], State.DescriptorSize);
+			FMemory::Memcpy(StagingBuffer->GetMappedPointer(), CPUDescriptor, State.DescriptorSize);
 			{
 				VkMemoryBarrier2 MemoryBarrier;
 				ZeroVulkanStruct(MemoryBarrier, VK_STRUCTURE_TYPE_MEMORY_BARRIER_2);
@@ -778,8 +842,6 @@ void FVulkanBindlessDescriptorManager::UpdateDescriptor(FRHIDescriptorHandle Des
 			Device.GetStagingManager().ReleaseBuffer(&Device.GetImmediateContext(), StagingBuffer);
 		}
 	}
-
-	UpdateStatsForHandle(DescriptorType);
 }
 
 void FVulkanBindlessDescriptorManager::UpdateSampler(FRHIDescriptorHandle DescriptorHandle, VkSampler VulkanSampler)
@@ -882,53 +944,6 @@ void FVulkanBindlessDescriptorManager::UpdateAccelerationStructure(FRHIDescripto
 	}
 }
 
-
-uint32 FVulkanBindlessDescriptorManager::GetFreeResourceIndex(FVulkanBindlessDescriptorManager::BindlessSetState& State)
-{
-	INC_DWORD_STAT(STAT_VulkanBindlessWritePerFrame);
-
-	{
-		FScopeLock ScopeLock(&State.FreeListCS);
-		if ((State.FreeListHead != MAX_uint32) && (State.PeakDescriptorCount >= State.MaxDescriptorCount)) // todo-jn: temp
-		{
-			const uint32 FreeIndex = State.FreeListHead;
-			const uint32 ByteOffset = State.FreeListHead * State.DescriptorSize;
-			uint32* NextSlot = (uint32*)(&State.DebugDescriptors[ByteOffset]);
-			State.FreeListHead = *NextSlot;
-			return FreeIndex;
-		}
-	}
-
-	INC_DWORD_STAT(STAT_VulkanBindlessPeakDescriptorCount);
-
-	const uint32 ResourceIndex = State.PeakDescriptorCount++;
-	checkf(ResourceIndex < State.MaxDescriptorCount, TEXT("You need to grow the resource array size for [%s]!"), VK_TYPE_TO_STRING(VkDescriptorType, State.DescriptorTypes[0]));
-	return ResourceIndex;
-}
-
-void FVulkanBindlessDescriptorManager::Unregister(FRHIDescriptorHandle DescriptorHandle)
-{
-	if (DescriptorHandle.IsValid())
-	{
-		checkf(bIsSupported, TEXT("Unregistering a valid handle but bindless is not supported!"));
-
-		const uint8 SetIndex = GetIndexForDescriptorType(DescriptorHandle.GetType());
-		BindlessSetState& State = BindlessSetStates[SetIndex];
-
-		FScopeLock ScopeLock(&State.FreeListCS);
-
-		const uint32 PreviousHead = State.FreeListHead;
-		State.FreeListHead = DescriptorHandle.GetIndex();
-		const uint32 ByteOffset = DescriptorHandle.GetIndex() * State.DescriptorSize;
-		uint32* NewSlot = (uint32*)(&State.DebugDescriptors[ByteOffset]);
-		FMemory::Memzero(NewSlot, State.DescriptorSize); // easier for debugging for now
-		*NewSlot = PreviousHead;
-
-		// Clear the descriptor
-		// todo-jn: invalidate the GPU side?
-	}
-}
-
 void FVulkanBindlessDescriptorManager::UpdateUBAllocator()
 {
 	if (bIsSupported && SingleUseUBAllocator)
diff --git a/Engine/Source/Runtime/VulkanRHI/Private/VulkanBindlessDescriptorManager.h b/Engine/Source/Runtime/VulkanRHI/Private/VulkanBindlessDescriptorManager.h
--- a/Engine/Source/Runtime/VulkanRHI/Private/VulkanBindlessDescriptorManager.h
+++ b/Engine/Source/Runtime/VulkanRHI/Private/VulkanBindlessDescriptorManager.h
@@ -7,8 +7,10 @@
 #include "VulkanCommon.h"
 #include "VulkanMemory.h"
 #include "VulkanThirdParty.h"
+#include "RHIDescriptorAllocator.h"
 
 class FVulkanDevice;
+class FVulkanBuffer;
 
 // Manager for resource descriptors used in bindless rendering.
 class FVulkanBindlessDescriptorManager
@@ -22,32 +24,38 @@ public:
 	void Init();
 	void Deinit();
 
-	static bool VerifySupport(FVulkanDevice& InDevice);
-	inline bool IsSupported()
+	bool IsSupported()
 	{
 		return bIsSupported;
 	}
 
-	inline VkPipelineLayout GetPipelineLayout() const
+	VkPipelineLayout GetPipelineLayout() const
 	{
 		return BindlessPipelineLayout;
 	}
 
 	void BindDescriptorBuffers(VkCommandBuffer CommandBuffer, VkPipelineStageFlags SupportedStages);
 
-	FRHIDescriptorHandle ReserveDescriptor(VkDescriptorType DescriptorType);
+	FRHIDescriptorHandle AllocateDescriptor(VkDescriptorType DescriptorType);
+	FRHIDescriptorHandle AllocateDescriptor(VkDescriptorType DescriptorType, VkImageView VulkanImage, bool bIsDepthStencil);
+	FRHIDescriptorHandle AllocateDescriptor(VkDescriptorType DescriptorType, VkDeviceAddress BufferAddress, VkDeviceSize BufferSize);
+	FRHIDescriptorHandle AllocateDescriptor(VkDescriptorType DescriptorType, FVulkanBuffer* Buffer, uint32 ExtraOffset);
+
+	FRHIDescriptorHandle AllocateDescriptor(VkSampler VulkanSampler);
+
+	void FreeDescriptor(FRHIDescriptorHandle DescriptorHandle);
 
 	void UpdateSampler(FRHIDescriptorHandle DescriptorHandle, VkSampler VulkanSampler);
 	void UpdateImage(FRHIDescriptorHandle DescriptorHandle, VkDescriptorType DescriptorType, VkImageView VulkanImage, bool bIsDepthStencil, bool bImmediateUpdate = true);
+
 	void UpdateBuffer(FRHIDescriptorHandle DescriptorHandle, VkDescriptorType DescriptorType, VkBuffer VulkanBuffer, VkDeviceSize BufferOffset, VkDeviceSize BufferSize, bool bImmediateUpdate = true);
 	void UpdateBuffer(FRHIDescriptorHandle DescriptorHandle, VkDescriptorType DescriptorType, VkDeviceAddress BufferAddress, VkDeviceSize BufferSize, bool bImmediateUpdate = true);
+
 	void UpdateTexelBuffer(FRHIDescriptorHandle DescriptorHandle, VkDescriptorType DescriptorType, const VkBufferViewCreateInfo& ViewInfo, bool bImmediateUpdate = true);
 	void UpdateAccelerationStructure(FRHIDescriptorHandle DescriptorHandle, VkAccelerationStructureKHR AccelerationStructure, bool bImmediateUpdate = true);
 
 	void RegisterUniformBuffers(FVulkanCommandListContext& Context, VkPipelineBindPoint BindPoint, const FUniformBufferDescriptorArrays& StageUBs);
 
-	void Unregister(FRHIDescriptorHandle DescriptorHandle);
-
 	void UpdateUBAllocator();
 
 private:
@@ -58,20 +66,27 @@ private:
 
 	struct BindlessSetState
 	{
-		TArray<VkDescriptorType> DescriptorTypes;
-		uint32 MaxDescriptorCount = 0;
-		std::atomic<uint32> PeakDescriptorCount = 1;  // always keep a null descriptor in slot 0
-		VkDescriptorSetLayout DescriptorSetLayout = VK_NULL_HANDLE;
+		~BindlessSetState()
+		{
+			if (CPUDescriptorMemory)
+			{
+				FMemory::Free(CPUDescriptorMemory);
+			}
+		}
+
+		TArray<VkDescriptorType, TInlineAllocator<2>> DescriptorTypes;
+		FRHIHeapDescriptorAllocator* Allocator = nullptr;
 
-		FCriticalSection FreeListCS;
-		uint32 FreeListHead = MAX_uint32;
+		VkDescriptorSetLayout DescriptorSetLayout = VK_NULL_HANDLE;
 
+		uint32 DescriptorCapacity = 0;
 		uint32 DescriptorSize = 0;
+
 		VkBuffer BufferHandle = VK_NULL_HANDLE;
 		VkDeviceMemory MemoryHandle = VK_NULL_HANDLE;
 		uint8* MappedPointer = nullptr;
 
-		TArray<uint8> DebugDescriptors;
+		uint8* CPUDescriptorMemory = nullptr;
 	};
 	BindlessSetState BindlessSetStates[VulkanBindless::NumBindlessSets];
 
@@ -83,7 +98,9 @@ private:
 
 	VkPipelineLayout BindlessPipelineLayout = VK_NULL_HANDLE;
 
-	uint32 GetFreeResourceIndex(BindlessSetState& Desc);
-	void UpdateStatsForHandle(VkDescriptorType DescriptorType);
+	BindlessSetState& GetBindlessState(ERHIDescriptorType DescriptorType);
+	BindlessSetState& GetBindlessState(VkDescriptorType DescriptorType);
+
 	void UpdateDescriptor(FRHIDescriptorHandle DescriptorHandle, VkDescriptorType DescriptorType, VkDescriptorDataEXT DescriptorData, bool bImmediateUpdate);
+	static bool VerifySupport(FVulkanDevice& InDevice);
 };
diff --git a/Engine/Source/Runtime/VulkanRHI/Private/VulkanDevice.cpp b/Engine/Source/Runtime/VulkanRHI/Private/VulkanDevice.cpp
--- a/Engine/Source/Runtime/VulkanRHI/Private/VulkanDevice.cpp
+++ b/Engine/Source/Runtime/VulkanRHI/Private/VulkanDevice.cpp
@@ -1696,7 +1696,7 @@ void FVulkanDevice::InitGlobalSamplers()
 	static_assert((uint32)FVulkanShaderHeader::EGlobalSamplerType::Count == 6, "Need to add new global sampler decl!");
 
 	// Make sure if bindless handles are being generated that these samplers are the very first
-	check(!GlobalSamplers[0]->GetBindlessHandle().IsValid() || GlobalSamplers[0]->GetBindlessHandle().GetIndex() == 1);
+	check(!GlobalSamplers[0]->GetBindlessHandle().IsValid() || GlobalSamplers[0]->GetBindlessHandle().GetIndex() == 0);
 }
 
 VkBuffer FVulkanDevice::CreateBuffer(VkDeviceSize BufferSize, VkBufferUsageFlags BufferUsageFlags, VkBufferCreateFlags BufferCreateFlags) const
diff --git a/Engine/Source/Runtime/VulkanRHI/Private/VulkanMemory.cpp b/Engine/Source/Runtime/VulkanRHI/Private/VulkanMemory.cpp
--- a/Engine/Source/Runtime/VulkanRHI/Private/VulkanMemory.cpp
+++ b/Engine/Source/Runtime/VulkanRHI/Private/VulkanMemory.cpp
@@ -4886,7 +4886,7 @@ namespace VulkanRHI
 			{
 				check(Device->SupportsBindless());
 				const FRHIDescriptorHandle& DescriptorHandle = reinterpret_cast<const FRHIDescriptorHandle&>(Entry.Handle);
-				Device->GetBindlessDescriptorManager()->Unregister(DescriptorHandle);
+				Device->GetBindlessDescriptorManager()->FreeDescriptor(DescriptorHandle);
 				break;
 			}
 
diff --git a/Engine/Source/Runtime/VulkanRHI/Private/VulkanRayTracing.cpp b/Engine/Source/Runtime/VulkanRHI/Private/VulkanRayTracing.cpp
--- a/Engine/Source/Runtime/VulkanRHI/Private/VulkanRayTracing.cpp
+++ b/Engine/Source/Runtime/VulkanRHI/Private/VulkanRayTracing.cpp
@@ -556,10 +556,7 @@ void FVulkanRayTracingGeometry::SetupHitGroupSystemParameters()
 	{
 		if (Buffer)
 		{
-			const VkDescriptorType DescriptorType = VK_DESCRIPTOR_TYPE_STORAGE_BUFFER;
-			FRHIDescriptorHandle BindlessHandle = BindlessDescriptorManager->ReserveDescriptor(DescriptorType);
-			BindlessDescriptorManager->UpdateBuffer(BindlessHandle, DescriptorType, Buffer->GetHandle(), Buffer->GetOffset() + ExtraOffset, Buffer->GetCurrentSize() - ExtraOffset);
-			return BindlessHandle;
+			return BindlessDescriptorManager->AllocateDescriptor(VK_DESCRIPTOR_TYPE_STORAGE_BUFFER, Buffer, ExtraOffset);
 		}
 		return FRHIDescriptorHandle();
 	};
@@ -598,13 +595,13 @@ void FVulkanRayTracingGeometry::ReleaseBindlessHandles()
 
 	for (FRHIDescriptorHandle BindlesHandle : HitGroupSystemVertexViews)
 	{
-		BindlessDescriptorManager->Unregister(BindlesHandle);
+		BindlessDescriptorManager->FreeDescriptor(BindlesHandle);
 	}
 	HitGroupSystemVertexViews.Reset(Initializer.Segments.Num());
 
 	if (HitGroupSystemIndexView.IsValid())
 	{
-		BindlessDescriptorManager->Unregister(HitGroupSystemIndexView);
+		BindlessDescriptorManager->FreeDescriptor(HitGroupSystemIndexView);
 		HitGroupSystemIndexView = FRHIDescriptorHandle();
 	}
 }
diff --git a/Engine/Source/Runtime/VulkanRHI/Private/VulkanState.cpp b/Engine/Source/Runtime/VulkanRHI/Private/VulkanState.cpp
--- a/Engine/Source/Runtime/VulkanRHI/Private/VulkanState.cpp
+++ b/Engine/Source/Runtime/VulkanRHI/Private/VulkanState.cpp
@@ -237,8 +237,7 @@ FVulkanSamplerState::FVulkanSamplerState(const VkSamplerCreateInfo& InInfo, FVul
 	if (InDevice.SupportsBindless() && (RHIGetRuntimeBindlessConfiguration(GMaxRHIShaderPlatform) != ERHIBindlessConfiguration::Disabled))
 	{
 		FVulkanBindlessDescriptorManager* BindlessDescriptorManager = InDevice.GetBindlessDescriptorManager();
-		BindlessHandle = BindlessDescriptorManager->ReserveDescriptor(VK_DESCRIPTOR_TYPE_SAMPLER);
-		BindlessDescriptorManager->UpdateSampler(BindlessHandle, Sampler);
+		BindlessHandle = BindlessDescriptorManager->AllocateDescriptor(Sampler);
 	}
 }
 
diff --git a/Engine/Source/Runtime/VulkanRHI/Private/VulkanUniformBuffer.cpp b/Engine/Source/Runtime/VulkanRHI/Private/VulkanUniformBuffer.cpp
--- a/Engine/Source/Runtime/VulkanRHI/Private/VulkanUniformBuffer.cpp
+++ b/Engine/Source/Runtime/VulkanRHI/Private/VulkanUniformBuffer.cpp
@@ -246,8 +246,7 @@ FRHIDescriptorHandle FVulkanUniformBuffer::GetBindlessHandle()
 			Device->GetDeferredDeletionQueue().EnqueueBindlessHandle(BindlessHandle);
 		}
 
-		BindlessHandle = Device->GetBindlessDescriptorManager()->ReserveDescriptor(VK_DESCRIPTOR_TYPE_UNIFORM_BUFFER);
-		Device->GetBindlessDescriptorManager()->UpdateBuffer(BindlessHandle, VK_DESCRIPTOR_TYPE_UNIFORM_BUFFER, CurrentAddress, GetSize(), true);
+		BindlessHandle = Device->GetBindlessDescriptorManager()->AllocateDescriptor(VK_DESCRIPTOR_TYPE_UNIFORM_BUFFER, CurrentAddress, GetSize());
 
 		CachedDeviceAddress = CurrentAddress;
 	}
diff --git a/Engine/Source/Runtime/VulkanRHI/Private/VulkanView.cpp b/Engine/Source/Runtime/VulkanRHI/Private/VulkanView.cpp
--- a/Engine/Source/Runtime/VulkanRHI/Private/VulkanView.cpp
+++ b/Engine/Source/Runtime/VulkanRHI/Private/VulkanView.cpp
@@ -88,7 +88,7 @@ FVulkanView::FVulkanView(FVulkanDevice& InDevice, VkDescriptorType InDescriptorT
 	: Device(InDevice)
 	, DescriptorType(InDescriptorType)
 {
-	BindlessHandle = Device.GetBindlessDescriptorManager()->ReserveDescriptor(InDescriptorType);
+	BindlessHandle = Device.GetBindlessDescriptorManager()->AllocateDescriptor(InDescriptorType);
 }
 
 FVulkanView::~FVulkanView()
