VulkanRHI: Add more extended dynamic state functions

#rnx

diff --git a/Engine/Source/Runtime/VulkanRHI/Private/VulkanCommandWrappers.h b/Engine/Source/Runtime/VulkanRHI/Private/VulkanCommandWrappers.h
--- a/Engine/Source/Runtime/VulkanRHI/Private/VulkanCommandWrappers.h
+++ b/Engine/Source/Runtime/VulkanRHI/Private/VulkanCommandWrappers.h
@@ -291,8 +291,13 @@ struct FWrapLayer
 	static void CmdSetDepthTestEnableEXT(VkResult Result, VkCommandBuffer CommandBuffer, VkBool32 DepthTestEnable) VULKAN_LAYER_BODY
 	static void CmdSetDepthWriteEnableEXT(VkResult Result, VkCommandBuffer CommandBuffer, VkBool32 DepthWriteEnable) VULKAN_LAYER_BODY
 	static void CmdSetRasterizerDiscardEnableEXT(VkResult Result, VkCommandBuffer CommandBuffer, VkBool32 RasterizerDiscardEnable) VULKAN_LAYER_BODY
+	static void CmdSetPrimitiveRestartEnableEXT(VkResult Result, VkCommandBuffer CommandBuffer, VkBool32 PrimitiveRestartEnable) VULKAN_LAYER_BODY
+	static void CmdSetLogicOpEnableEXT(VkResult Result, VkCommandBuffer CommandBuffer, VkBool32 LogicOpEnable) VULKAN_LAYER_BODY
+	static void CmdSetAlphaToCoverageEnableEXT(VkResult Result, VkCommandBuffer CommandBuffer, VkBool32 AlphaToCoverageEnable) VULKAN_LAYER_BODY
+	static void CmdSetAlphaToOneEnableEXT(VkResult Result, VkCommandBuffer CommandBuffer, VkBool32 AlphaToOneEnable) VULKAN_LAYER_BODY
 };
 
+
 #undef VULKAN_LAYER_BODY
 
 namespace VulkanRHI
@@ -2044,6 +2049,34 @@ namespace VulkanRHI
 		FWrapLayer::CmdSetRasterizerDiscardEnableEXT(VK_SUCCESS, CommandBuffer, RasterizerDiscardEnable);
 	}
 
+	static FORCEINLINE_DEBUGGABLE void vkCmdSetPrimitiveRestartEnableEXT(VkCommandBuffer CommandBuffer, VkBool32 PrimitiveRestartEnable)
+	{
+		FWrapLayer::CmdSetPrimitiveRestartEnableEXT(VK_RESULT_MAX_ENUM, CommandBuffer, PrimitiveRestartEnable);
+		VULKANAPINAMESPACE::vkCmdSetPrimitiveRestartEnableEXT(CommandBuffer, PrimitiveRestartEnable);
+		FWrapLayer::CmdSetPrimitiveRestartEnableEXT(VK_SUCCESS, CommandBuffer, PrimitiveRestartEnable);
+	}
+
+	static FORCEINLINE_DEBUGGABLE void vkCmdSetLogicOpEnableEXT(VkCommandBuffer CommandBuffer, VkBool32 LogicOpEnable)
+	{
+		FWrapLayer::CmdSetLogicOpEnableEXT(VK_RESULT_MAX_ENUM, CommandBuffer, LogicOpEnable);
+		VULKANAPINAMESPACE::vkCmdSetLogicOpEnableEXT(CommandBuffer, LogicOpEnable);
+		FWrapLayer::CmdSetLogicOpEnableEXT(VK_SUCCESS, CommandBuffer, LogicOpEnable);
+	}
+
+	static FORCEINLINE_DEBUGGABLE void vkCmdSetAlphaToCoverageEnableEXT(VkCommandBuffer CommandBuffer, VkBool32 AlphaToCoverageEnable)
+	{
+		FWrapLayer::CmdSetAlphaToCoverageEnableEXT(VK_RESULT_MAX_ENUM, CommandBuffer, AlphaToCoverageEnable);
+		VULKANAPINAMESPACE::vkCmdSetAlphaToCoverageEnableEXT(CommandBuffer, AlphaToCoverageEnable);
+		FWrapLayer::CmdSetAlphaToCoverageEnableEXT(VK_SUCCESS, CommandBuffer, AlphaToCoverageEnable);
+	}
+
+	static FORCEINLINE_DEBUGGABLE void vkCmdSetAlphaToOneEnableEXT(VkCommandBuffer CommandBuffer, VkBool32 AlphaToOneEnable)
+	{
+		FWrapLayer::CmdSetAlphaToOneEnableEXT(VK_RESULT_MAX_ENUM, CommandBuffer, AlphaToOneEnable);
+		VULKANAPINAMESPACE::vkCmdSetAlphaToOneEnableEXT(CommandBuffer, AlphaToOneEnable);
+		FWrapLayer::CmdSetAlphaToOneEnableEXT(VK_SUCCESS, CommandBuffer, AlphaToOneEnable);
+	}
+
 
 #if VULKAN_ENABLE_IMAGE_TRACKING_LAYER
 	void BindDebugLabelName(VkImage Image, const TCHAR* Name);
diff --git a/Engine/Source/Runtime/VulkanRHI/Private/VulkanDebug.cpp b/Engine/Source/Runtime/VulkanRHI/Private/VulkanDebug.cpp
--- a/Engine/Source/Runtime/VulkanRHI/Private/VulkanDebug.cpp
+++ b/Engine/Source/Runtime/VulkanRHI/Private/VulkanDebug.cpp
@@ -4955,6 +4955,45 @@ void FWrapLayer::CmdSetRasterizerDiscardEnableEXT(VkResult Result, VkCommandBuff
 	}
 }
 
+void FWrapLayer::CmdSetPrimitiveRestartEnableEXT(VkResult Result, VkCommandBuffer CommandBuffer, VkBool32 PrimitiveRestartEnable)
+{
+	if (Result == VK_RESULT_MAX_ENUM)
+	{
+#if VULKAN_ENABLE_DUMP_LAYER
+		PrintfBeginResult(FString::Printf(TEXT("CmdSetPrimitiveRestartEnableEXT()")));
+#endif
+	}
+}
+
+void FWrapLayer::CmdSetLogicOpEnableEXT(VkResult Result, VkCommandBuffer CommandBuffer, VkBool32 LogicOpEnable)
+{
+	if (Result == VK_RESULT_MAX_ENUM)
+	{
+#if VULKAN_ENABLE_DUMP_LAYER
+		PrintfBeginResult(FString::Printf(TEXT("CmdSetLogicOpEnableEXT()")));
+#endif
+	}
+}
+
+void FWrapLayer::CmdSetAlphaToCoverageEnableEXT(VkResult Result, VkCommandBuffer CommandBuffer, VkBool32 AlphaToCoverageEnable)
+{
+	if (Result == VK_RESULT_MAX_ENUM)
+	{
+#if VULKAN_ENABLE_DUMP_LAYER
+		PrintfBeginResult(FString::Printf(TEXT("CmdSetAlphaToCoverageEnableEXT()")));
+#endif
+	}
+}
+
+void FWrapLayer::CmdSetAlphaToOneEnableEXT(VkResult Result, VkCommandBuffer CommandBuffer, VkBool32 AlphaToOneEnable)
+{
+	if (Result == VK_RESULT_MAX_ENUM)
+	{
+#if VULKAN_ENABLE_DUMP_LAYER
+		PrintfBeginResult(FString::Printf(TEXT("CmdSetAlphaToOneEnableEXT()")));
+#endif
+	}
+}
 
 
 
diff --git a/Engine/Source/Runtime/VulkanRHI/Private/VulkanLoader.h b/Engine/Source/Runtime/VulkanRHI/Private/VulkanLoader.h
--- a/Engine/Source/Runtime/VulkanRHI/Private/VulkanLoader.h
+++ b/Engine/Source/Runtime/VulkanRHI/Private/VulkanLoader.h
@@ -266,7 +266,12 @@
 	EnumMacro(PFN_vkCmdSetDepthCompareOpEXT, vkCmdSetDepthCompareOpEXT) \
 	EnumMacro(PFN_vkCmdSetDepthTestEnableEXT, vkCmdSetDepthTestEnableEXT) \
 	EnumMacro(PFN_vkCmdSetDepthWriteEnableEXT, vkCmdSetDepthWriteEnableEXT) \
-	EnumMacro(PFN_vkCmdSetRasterizerDiscardEnableEXT, vkCmdSetRasterizerDiscardEnableEXT)
+	EnumMacro(PFN_vkCmdSetRasterizerDiscardEnableEXT, vkCmdSetRasterizerDiscardEnableEXT) \
+	EnumMacro(PFN_vkCmdSetPrimitiveRestartEnableEXT, vkCmdSetPrimitiveRestartEnableEXT) \
+	EnumMacro(PFN_vkCmdSetAlphaToCoverageEnableEXT, vkCmdSetAlphaToCoverageEnableEXT) \
+	EnumMacro(PFN_vkCmdSetAlphaToOneEnableEXT, vkCmdSetAlphaToOneEnableEXT) \
+	EnumMacro(PFN_vkCmdSetLogicOpEnableEXT, vkCmdSetLogicOpEnableEXT)
+
 
 
 
