VulkanRHI: Add all extended dynamic state extensions (WIP, no dedpulication, disabled by default).

#rnx

diff --git a/Engine/Source/Runtime/VulkanRHI/Private/VulkanCommandBuffer.cpp b/Engine/Source/Runtime/VulkanRHI/Private/VulkanCommandBuffer.cpp
--- a/Engine/Source/Runtime/VulkanRHI/Private/VulkanCommandBuffer.cpp
+++ b/Engine/Source/Runtime/VulkanRHI/Private/VulkanCommandBuffer.cpp
@@ -12,6 +12,7 @@
 #include "VulkanMemory.h"
 #include "VulkanRayTracing.h"
 #include "VulkanBindlessDescriptorManager.h"
+#include "VulkanPipeline.h"
 
 
 #define CMD_BUFFER_TIME_TO_WAIT_BEFORE_DELETING		10
@@ -26,6 +27,12 @@ FVulkanCommandBuffer::FVulkanCommandBuffer(FVulkanDevice& InDevice, FVulkanComma
 {
 	FScopeLock ScopeLock(CommandBufferPool.GetCS());
 	AllocMemory();
+
+	if (InDevice.GetOptionalExtensions().HasAnyExtendedDynamicState())
+	{
+		LastDynamicStateUpdate = new FGfxPipelineDesc();
+		FMemory::Memzero(*LastDynamicStateUpdate);
+	}
 }
 
 void FVulkanCommandBuffer::AllocMemory()
@@ -43,6 +50,7 @@ void FVulkanCommandBuffer::AllocMemory()
 
 	VERIFYVULKANRESULT(VulkanRHI::vkAllocateCommandBuffers(Device.GetHandle(), &CreateCmdBufInfo, &CommandBufferHandle));
 
+	bNeedsFullDynamicStateUpdate = 1;
 	bHasPipeline = 0;
 	bHasViewport = 0;
 	bHasScissor = 0;
@@ -58,6 +66,12 @@ FVulkanCommandBuffer::~FVulkanCommandBuffer()
 	{
 		FreeMemory();
 	}
+
+	if (LastDynamicStateUpdate)
+	{
+		delete LastDynamicStateUpdate;
+		LastDynamicStateUpdate = nullptr;
+	}
 }
 
 void FVulkanCommandBuffer::FreeMemory()
@@ -215,6 +229,8 @@ void FVulkanCommandBuffer::Begin(FVulkanQueryPool* OptionalQueryPool, VkRenderPa
 		VulkanRHI::vkCmdWriteTimestamp(GetHandle(), VK_PIPELINE_STAGE_TOP_OF_PIPE_BIT, OptionalQueryPool->GetHandle(), IndexInPool);
 	}
 #endif
+
+	bNeedsFullDynamicStateUpdate = true;
 }
 
 void FVulkanCommandBuffer::Reset()
@@ -230,6 +246,7 @@ void FVulkanCommandBuffer::Reset()
 	FScopeLock ScopeLock(CommandBufferPool.GetCS());
 	if (State == EState::Submitted)
 	{
+		bNeedsFullDynamicStateUpdate = true;
 		bHasPipeline = false;
 		bHasViewport = false;
 		bHasScissor = false;
diff --git a/Engine/Source/Runtime/VulkanRHI/Private/VulkanCommandBuffer.h b/Engine/Source/Runtime/VulkanRHI/Private/VulkanCommandBuffer.h
--- a/Engine/Source/Runtime/VulkanRHI/Private/VulkanCommandBuffer.h
+++ b/Engine/Source/Runtime/VulkanRHI/Private/VulkanCommandBuffer.h
@@ -8,17 +8,15 @@
 
 #include "CoreMinimal.h"
 #include "VulkanConfiguration.h"
-#include "VulkanBarriers.h"
+#include "VulkanThirdParty.h"
 
 class FVulkanDevice;
 class FVulkanCommandBufferPool;
 class FVulkanContextCommon;
-class FVulkanRenderTargetLayout;
 class FVulkanQueue;
-class FVulkanDescriptorPoolSetContainer;
 class FVulkanQueryPool;
 struct FVulkanBeginRenderPassInfo;
-struct FVulkanDynamicRenderingInfo;
+struct FGfxPipelineDesc;
 
 namespace VulkanRHI
 {
@@ -108,6 +106,10 @@ public:
 	TArray<VkRect2D, TInlineAllocator<2>> CurrentScissors;
 	uint32 CurrentStencilRef = 0;
 	EState State = EState::NotAllocated;
+
+	FGfxPipelineDesc* LastDynamicStateUpdate = nullptr;
+
+	uint8 bNeedsFullDynamicStateUpdate	: 1 = 1;
 	uint8 bHasPipeline					: 1 = 0;
 	uint8 bHasViewport					: 1 = 0;
 	uint8 bHasScissor					: 1 = 0;
diff --git a/Engine/Source/Runtime/VulkanRHI/Private/VulkanDevice.h b/Engine/Source/Runtime/VulkanRHI/Private/VulkanDevice.h
--- a/Engine/Source/Runtime/VulkanRHI/Private/VulkanDevice.h
+++ b/Engine/Source/Runtime/VulkanRHI/Private/VulkanDevice.h
@@ -116,12 +116,12 @@ struct FOptionalVulkanDeviceExtensions
 		Packed = 0;
 	}
 
-	inline bool HasGPUCrashDumpExtensions() const
+	bool HasGPUCrashDumpExtensions() const
 	{
 		return HasAMDBufferMarker || HasNVDiagnosticCheckpoints;
 	}
 
-	inline bool HasRaytracingExtensions() const
+	bool HasRaytracingExtensions() const
 	{
 		return 
 			HasAccelerationStructure && 
@@ -132,6 +132,11 @@ struct FOptionalVulkanDeviceExtensions
 			HasSPIRV_14 && 
 			HasShaderFloatControls;
 	}
+
+	bool HasAnyExtendedDynamicState() const
+	{
+		return HasEXTExtendedDynamicState1 || HasEXTExtendedDynamicState2 || HasEXTExtendedDynamicState3;
+	}
 };
 
 // All the features and properties we need to keep around from extension initialization
diff --git a/Engine/Source/Runtime/VulkanRHI/Private/VulkanPendingState.cpp b/Engine/Source/Runtime/VulkanRHI/Private/VulkanPendingState.cpp
--- a/Engine/Source/Runtime/VulkanRHI/Private/VulkanPendingState.cpp
+++ b/Engine/Source/Runtime/VulkanRHI/Private/VulkanPendingState.cpp
@@ -390,6 +390,13 @@ void FVulkanPendingGfxState::PrepareForDraw(FVulkanCommandListContext& Context)
 
 	check(Context.GetCommandBuffer().bHasPipeline);
 
+	// Force an update on vertex stream strides
+	if (Context.GetCommandBuffer().bNeedsFullDynamicStateUpdate && 
+		Device.GetOptionalExtensions().HasEXTExtendedDynamicState1)
+	{
+		bDirtyDynamicStrides = true;
+	}
+
 	if (CurrentPipeline->UsesBindless())
 	{
 		check(!CurrentPipeline->bHasInputAttachments); // todo-jn: bindless + InputAttachments
@@ -415,7 +422,7 @@ void FVulkanPendingGfxState::PrepareForDraw(FVulkanCommandListContext& Context)
 		}
 	}
 
-	if (bDirtyVertexStreams)
+	if (bDirtyVertexStreams || bDirtyDynamicStrides)
 	{
 #if VULKAN_ENABLE_AGGRESSIVE_STATS
 		SCOPE_CYCLE_COUNTER(STAT_VulkanBindVertexStreamsTime);
@@ -426,6 +433,10 @@ void FVulkanPendingGfxState::PrepareForDraw(FVulkanCommandListContext& Context)
 		{
 			// However, we need to verify that there are also no bindings
 			check(VertexInputStateInfo.BindingsNum == 0);
+
+			// Will clear the dynamic strides, be sure to set them on the next pass
+			bDirtyDynamicStrides = true;
+
 			return;
 		}
 
@@ -433,13 +444,17 @@ void FVulkanPendingGfxState::PrepareForDraw(FVulkanCommandListContext& Context)
 		{
 			VkBuffer VertexBuffers[MaxVertexElementCount];
 			VkDeviceSize VertexOffsets[MaxVertexElementCount];
+			VkDeviceSize VertexSizes[MaxVertexElementCount];
+			VkDeviceSize VertexStrides[MaxVertexElementCount];
 			int32 NumUsed = 0;
 
-			void Add(VkBuffer InBuffer, VkDeviceSize InSize)
+			void Add(VkBuffer InBuffer, VkDeviceSize InOffset, VkDeviceSize InStride)
 			{
 				check(NumUsed < MaxVertexElementCount);
 				VertexBuffers[NumUsed] = InBuffer;
-				VertexOffsets[NumUsed] = InSize;
+				VertexOffsets[NumUsed] = InOffset;
+				VertexSizes[NumUsed] = VK_WHOLE_SIZE;
+				VertexStrides[NumUsed] = InStride;
 				++NumUsed;
 			}
 		} TemporaryIA;
@@ -474,7 +489,7 @@ void FVulkanPendingGfxState::PrepareForDraw(FVulkanCommandListContext& Context)
 				continue;
 			}
 
-			TemporaryIA.Add(CurrStream.Stream, CurrStream.BufferOffset);
+			TemporaryIA.Add(CurrStream.Stream, CurrStream.BufferOffset, CurrBinding.stride);
 		}
 
 		if (TemporaryIA.NumUsed > 0)
@@ -484,7 +499,19 @@ void FVulkanPendingGfxState::PrepareForDraw(FVulkanCommandListContext& Context)
 			// Incorrect:	1, 0, 2, 3
 			// Incorrect:	0, 2, 3, 5
 			// Reordering and creation of stream binding index is done in "GenerateVertexInputStateInfo()"
-			VulkanRHI::vkCmdBindVertexBuffers(Context.GetCommandBuffer().GetHandle(), 0, TemporaryIA.NumUsed, TemporaryIA.VertexBuffers, TemporaryIA.VertexOffsets);
+			if (Device.GetOptionalExtensions().HasEXTExtendedDynamicState1)
+			{
+				VulkanRHI::vkCmdBindVertexBuffers2EXT(Context.GetCommandBuffer().GetHandle(), 0, 
+					TemporaryIA.NumUsed, TemporaryIA.VertexBuffers, TemporaryIA.VertexOffsets, TemporaryIA.VertexSizes, 
+					CurrentPipeline->GetVulkanShader(SF_Mesh) ? nullptr : TemporaryIA.VertexStrides);
+
+				// Mesh shader will clear the dynamic strides, be sure to set them on the next pass
+				bDirtyDynamicStrides = (CurrentPipeline->GetVulkanShader(SF_Mesh) == nullptr);
+			}
+			else
+			{
+				VulkanRHI::vkCmdBindVertexBuffers(Context.GetCommandBuffer().GetHandle(), 0, TemporaryIA.NumUsed, TemporaryIA.VertexBuffers, TemporaryIA.VertexOffsets);
+			}
 		}
 
 		bDirtyVertexStreams = false;
@@ -508,7 +535,15 @@ void FVulkanPendingGfxState::UpdateDynamicStates(FVulkanCommandBuffer& CommandBu
 			FlippedViewport.y += FlippedViewport.height;
 			FlippedViewport.height = -FlippedViewport.height;
 		}
-		VulkanRHI::vkCmdSetViewport(CommandBuffer.GetHandle(), 0, FlippedViewports.Num(), FlippedViewports.GetData());
+
+		if (Device.GetOptionalExtensions().HasEXTExtendedDynamicState1)
+		{
+			VulkanRHI::vkCmdSetViewportWithCountEXT(CommandBuffer.GetHandle(), FlippedViewports.Num(), FlippedViewports.GetData());
+		}
+		else
+		{
+			VulkanRHI::vkCmdSetViewport(CommandBuffer.GetHandle(), 0, FlippedViewports.Num(), FlippedViewports.GetData());
+		}
 
 		CommandBuffer.CurrentViewports = Viewports;
 		CommandBuffer.bHasViewport = true;
@@ -518,7 +553,15 @@ void FVulkanPendingGfxState::UpdateDynamicStates(FVulkanCommandBuffer& CommandBu
 		(FMemory::Memcmp((const void*)CommandBuffer.CurrentScissors.GetData(), (const void*)Scissors.GetData(), Scissors.Num() * sizeof(VkRect2D)) != 0);
 	if (bNeedsUpdateScissor)
 	{
-		VulkanRHI::vkCmdSetScissor(CommandBuffer.GetHandle(), 0, Scissors.Num(), Scissors.GetData());
+		if (Device.GetOptionalExtensions().HasEXTExtendedDynamicState1)
+		{
+			VulkanRHI::vkCmdSetScissorWithCountEXT(CommandBuffer.GetHandle(), Scissors.Num(), Scissors.GetData());
+		}
+		else
+		{
+			VulkanRHI::vkCmdSetScissor(CommandBuffer.GetHandle(), 0, Scissors.Num(), Scissors.GetData());
+		}
+
 		CommandBuffer.CurrentScissors = Scissors;
 		CommandBuffer.bHasScissor = true;
 	}
@@ -530,6 +573,199 @@ void FVulkanPendingGfxState::UpdateDynamicStates(FVulkanCommandBuffer& CommandBu
 		CommandBuffer.CurrentStencilRef = StencilRef;
 		CommandBuffer.bHasStencilRef = true;
 	}
+
+	const bool bForceUpdate = CommandBuffer.bNeedsFullDynamicStateUpdate;
+	FGfxPipelineDesc* LastDynamicStateUpdate = CommandBuffer.LastDynamicStateUpdate;
+	const FGfxPipelineDesc& NewDynamicState = CurrentPipeline->Desc;
+
+	// Bunch of macros to make dynamic state comparions below easier to read
+#define VULKAN_DYNAMIC_STATE_COMPARE(Name) (LastDynamicStateUpdate->Name != NewDynamicState.Name)
+#define VULKAN_DYNAMIC_STATE_COMPARE1(Name, ...) VULKAN_DYNAMIC_STATE_COMPARE(Name)
+#define VULKAN_DYNAMIC_STATE_COMPARE2(Name, ...) VULKAN_DYNAMIC_STATE_COMPARE(Name) __VA_OPT__( || VULKAN_DYNAMIC_STATE_COMPARE1(__VA_ARGS__))
+#define VULKAN_DYNAMIC_STATE_COMPARE3(Name, ...) VULKAN_DYNAMIC_STATE_COMPARE(Name) __VA_OPT__( || VULKAN_DYNAMIC_STATE_COMPARE2(__VA_ARGS__))
+
+#define VULKAN_DYNAMIC_STATE_UPDATE(Name) LastDynamicStateUpdate->Name = NewDynamicState.Name; 
+#define VULKAN_DYNAMIC_STATE_UPDATE1(Name, ...) VULKAN_DYNAMIC_STATE_UPDATE(Name)
+#define VULKAN_DYNAMIC_STATE_UPDATE2(Name, ...) VULKAN_DYNAMIC_STATE_UPDATE(Name) __VA_OPT__( VULKAN_DYNAMIC_STATE_UPDATE1(__VA_ARGS__))
+#define VULKAN_DYNAMIC_STATE_UPDATE3(Name, ...) VULKAN_DYNAMIC_STATE_UPDATE(Name) __VA_OPT__( VULKAN_DYNAMIC_STATE_UPDATE2(__VA_ARGS__))
+
+#define VULKAN_DYNAMIC_STATE_BLOCK_START(...) if (bForceUpdate || VULKAN_DYNAMIC_STATE_COMPARE3(__VA_ARGS__)) { VULKAN_DYNAMIC_STATE_UPDATE3(__VA_ARGS__)
+#define VULKAN_DYNAMIC_STATE_BLOCK_END()      }
+
+	// NOTE: All of the different VK_DYNAMIC_STATE_* values being set can be found in VulkanPipeline.cpp
+
+	if (Device.GetOptionalExtensions().HasEXTExtendedDynamicState1)
+	{
+		check(LastDynamicStateUpdate);
+
+		if (CurrentPipeline->GetVulkanShader(SF_Mesh) == nullptr)
+		{
+			// VK_DYNAMIC_STATE_PRIMITIVE_TOPOLOGY
+			VULKAN_DYNAMIC_STATE_BLOCK_START(Topology);
+			VulkanRHI::vkCmdSetPrimitiveTopologyEXT(CommandBuffer.GetHandle(), (VkPrimitiveTopology)NewDynamicState.Topology);
+			VULKAN_DYNAMIC_STATE_BLOCK_END();
+		}
+		else
+		{
+			// The state will be invalidated, make sure we set it again
+			LastDynamicStateUpdate->Topology = VK_PRIMITIVE_TOPOLOGY_MAX_ENUM;
+		}
+
+		// VK_DYNAMIC_STATE_CULL_MODE
+		VULKAN_DYNAMIC_STATE_BLOCK_START(Rasterizer.CullMode);
+		VulkanRHI::vkCmdSetCullModeEXT(CommandBuffer.GetHandle(), (VkCullModeFlags)NewDynamicState.Rasterizer.CullMode);
+		VULKAN_DYNAMIC_STATE_BLOCK_END();
+
+		// VK_DYNAMIC_STATE_STENCIL_TEST_ENABLE
+		VULKAN_DYNAMIC_STATE_BLOCK_START(DepthStencil.bStencilTestEnable);
+		VulkanRHI::vkCmdSetStencilTestEnableEXT(CommandBuffer.GetHandle(), NewDynamicState.DepthStencil.bStencilTestEnable ? VK_TRUE : VK_FALSE);
+		VULKAN_DYNAMIC_STATE_BLOCK_END();
+
+		// VK_DYNAMIC_STATE_STENCIL_OP (front)
+		VULKAN_DYNAMIC_STATE_BLOCK_START(DepthStencil.FrontFailOp, DepthStencil.FrontPassOp, DepthStencil.FrontDepthFailOp, DepthStencil.FrontCompareOp);
+		VulkanRHI::vkCmdSetStencilOpEXT(
+			CommandBuffer.GetHandle(),
+			VK_STENCIL_FACE_FRONT_BIT,
+			(VkStencilOp)NewDynamicState.DepthStencil.FrontFailOp,
+			(VkStencilOp)NewDynamicState.DepthStencil.FrontPassOp,
+			(VkStencilOp)NewDynamicState.DepthStencil.FrontDepthFailOp,
+			(VkCompareOp)NewDynamicState.DepthStencil.FrontCompareOp
+		);
+		VULKAN_DYNAMIC_STATE_BLOCK_END();
+
+		// VK_DYNAMIC_STATE_STENCIL_OP (back)
+		VULKAN_DYNAMIC_STATE_BLOCK_START(DepthStencil.BackFailOp, DepthStencil.BackPassOp, DepthStencil.BackDepthFailOp, DepthStencil.BackCompareOp);
+		VulkanRHI::vkCmdSetStencilOpEXT(
+			CommandBuffer.GetHandle(),
+			VK_STENCIL_FACE_BACK_BIT,
+			(VkStencilOp)NewDynamicState.DepthStencil.BackFailOp,
+			(VkStencilOp)NewDynamicState.DepthStencil.BackPassOp,
+			(VkStencilOp)NewDynamicState.DepthStencil.BackDepthFailOp,
+			(VkCompareOp)NewDynamicState.DepthStencil.BackCompareOp
+		);
+		VULKAN_DYNAMIC_STATE_BLOCK_END();
+
+		// VK_DYNAMIC_STATE_DEPTH_COMPARE_OP
+		VULKAN_DYNAMIC_STATE_BLOCK_START(DepthStencil.DepthCompareOp);
+		VulkanRHI::vkCmdSetDepthCompareOpEXT(CommandBuffer.GetHandle(), (VkCompareOp)NewDynamicState.DepthStencil.DepthCompareOp);
+		VULKAN_DYNAMIC_STATE_BLOCK_END();
+
+		// VK_DYNAMIC_STATE_DEPTH_TEST_ENABLE
+		VULKAN_DYNAMIC_STATE_BLOCK_START(DepthStencil.bDepthTestEnable);
+		VulkanRHI::vkCmdSetDepthTestEnableEXT(CommandBuffer.GetHandle(), NewDynamicState.DepthStencil.bDepthTestEnable ? VK_TRUE : VK_FALSE);
+		VULKAN_DYNAMIC_STATE_BLOCK_END();
+
+		// VK_DYNAMIC_STATE_DEPTH_BOUNDS_TEST_ENABLE
+		VULKAN_DYNAMIC_STATE_BLOCK_START(DepthStencil.bDepthBoundsTestEnable);
+		VulkanRHI::vkCmdSetDepthBoundsTestEnableEXT(CommandBuffer.GetHandle(), NewDynamicState.DepthStencil.bDepthBoundsTestEnable ? VK_TRUE : VK_FALSE);
+		VULKAN_DYNAMIC_STATE_BLOCK_END();
+
+		// VK_DYNAMIC_STATE_DEPTH_WRITE_ENABLE
+		VULKAN_DYNAMIC_STATE_BLOCK_START(DepthStencil.bDepthWriteEnable);
+		VulkanRHI::vkCmdSetDepthWriteEnableEXT(CommandBuffer.GetHandle(), NewDynamicState.DepthStencil.bDepthWriteEnable ? VK_TRUE : VK_FALSE);
+		VULKAN_DYNAMIC_STATE_BLOCK_END();
+	}
+
+	if (Device.GetOptionalExtensions().HasEXTExtendedDynamicState2)
+	{
+		// VK_DYNAMIC_STATE_DEPTH_BIAS_ENABLE
+		const float OldDepthBiasConstantFactor = NewDynamicState.Rasterizer.DepthBiasConstantFactor;
+		VULKAN_DYNAMIC_STATE_BLOCK_START(Rasterizer.DepthBiasConstantFactor);
+		VulkanRHI::vkCmdSetDepthBiasEnableEXT(CommandBuffer.GetHandle(), (NewDynamicState.Rasterizer.DepthBiasConstantFactor != 0.0f) ? VK_TRUE : VK_FALSE);
+		// Put back the old value if VK_DYNAMIC_STATE_DEPTH_BIAS is int use so that it gets updated
+		// Note: this does mean we call enable for each DepthBiasConstantFactor
+		if (Device.GetOptionalExtensions().HasEXTExtendedDynamicState3)
+		{
+			LastDynamicStateUpdate->Rasterizer.DepthBiasConstantFactor = OldDepthBiasConstantFactor;
+		}
+		VULKAN_DYNAMIC_STATE_BLOCK_END();
+	}
+
+	if (Device.GetOptionalExtensions().HasEXTExtendedDynamicState3)
+	{
+		// VK_DYNAMIC_STATE_POLYGON_MODE_EXT
+		VULKAN_DYNAMIC_STATE_BLOCK_START(Rasterizer.PolygonMode);
+		VulkanRHI::vkCmdSetPolygonModeEXT(CommandBuffer.GetHandle(), (VkPolygonMode)NewDynamicState.Rasterizer.PolygonMode);
+		VULKAN_DYNAMIC_STATE_BLOCK_END();
+
+		// VK_DYNAMIC_STATE_RASTERIZATION_SAMPLES_EXT
+		VULKAN_DYNAMIC_STATE_BLOCK_START(RasterizationSamples);
+		const VkSampleCountFlagBits RasterizationSamples = (VkSampleCountFlagBits)FMath::Max<uint16>(1u, NewDynamicState.RasterizationSamples);
+		VulkanRHI::vkCmdSetRasterizationSamplesEXT(CommandBuffer.GetHandle(), RasterizationSamples);
+		VULKAN_DYNAMIC_STATE_BLOCK_END();
+
+		//VK_DYNAMIC_STATE_ALPHA_TO_COVERAGE_ENABLE_EXT
+		VULKAN_DYNAMIC_STATE_BLOCK_START(UseAlphaToCoverage);
+		VulkanRHI::vkCmdSetAlphaToCoverageEnableEXT(CommandBuffer.GetHandle(), NewDynamicState.UseAlphaToCoverage ? VK_TRUE : VK_FALSE);
+		VULKAN_DYNAMIC_STATE_BLOCK_END();
+
+		// VK_DYNAMIC_STATE_COLOR_BLEND_ENABLE_EXT
+		// VK_DYNAMIC_STATE_COLOR_BLEND_EQUATION_EXT
+		// VK_DYNAMIC_STATE_COLOR_WRITE_MASK_EXT
+		if (bForceUpdate || 
+			(LastDynamicStateUpdate->ShaderKeyShared != NewDynamicState.ShaderKeyShared) ||
+			(LastDynamicStateUpdate->ColorAttachmentStates != NewDynamicState.ColorAttachmentStates))
+		{
+			if (NewDynamicState.ColorAttachmentStates.Num())
+			{
+				check(NewDynamicState.ColorAttachmentStates.Num() <= MaxSimultaneousRenderTargets);
+
+				VkBool32 BlendEnabled[MaxSimultaneousRenderTargets];
+				VkColorBlendEquationEXT ColorBlendEquation[MaxSimultaneousRenderTargets];
+				VkColorComponentFlags ColorWriteMask[MaxSimultaneousRenderTargets];
+				uint32 TargetWriteMask = 0xffffffff;
+				if (const FVulkanShader* PixelShader = CurrentPipeline->GetVulkanShader(SF_Pixel))
+				{
+					TargetWriteMask = PixelShader->GetCodeHeader().InOutMask;
+				}
+				for (int32 Index = 0; Index < NewDynamicState.ColorAttachmentStates.Num(); ++Index)
+				{
+					const FGfxPipelineDesc::FBlendAttachment& BlendAttachment = NewDynamicState.ColorAttachmentStates[Index];
+
+					BlendEnabled[Index] = BlendAttachment.bBlend ? VK_TRUE : VK_FALSE;
+
+					ColorBlendEquation[Index].colorBlendOp = (VkBlendOp)BlendAttachment.ColorBlendOp;
+					ColorBlendEquation[Index].srcColorBlendFactor = (VkBlendFactor)BlendAttachment.SrcColorBlendFactor;
+					ColorBlendEquation[Index].dstColorBlendFactor = (VkBlendFactor)BlendAttachment.DstColorBlendFactor;
+					ColorBlendEquation[Index].alphaBlendOp = (VkBlendOp)BlendAttachment.AlphaBlendOp;
+					ColorBlendEquation[Index].srcAlphaBlendFactor = (VkBlendFactor)BlendAttachment.SrcAlphaBlendFactor;
+					ColorBlendEquation[Index].dstAlphaBlendFactor = (VkBlendFactor)BlendAttachment.DstAlphaBlendFactor;
+
+					ColorWriteMask[Index] = (VkColorComponentFlags)BlendAttachment.ColorWriteMask;
+					if (0 == (TargetWriteMask & 1)) //clear write mask of rendertargets not written by pixelshader.
+					{
+						ColorWriteMask[Index] = 0;
+					}
+					TargetWriteMask >>= 1;
+				}
+
+				VulkanRHI::vkCmdSetColorBlendEnableEXT(
+					CommandBuffer.GetHandle(),
+					0, NewDynamicState.ColorAttachmentStates.Num(),
+					BlendEnabled);
+
+				VulkanRHI::vkCmdSetColorBlendEquationEXT(
+					CommandBuffer.GetHandle(),
+					0, NewDynamicState.ColorAttachmentStates.Num(),
+					ColorBlendEquation);
+
+				VulkanRHI::vkCmdSetColorWriteMaskEXT(
+					CommandBuffer.GetHandle(),
+					0, NewDynamicState.ColorAttachmentStates.Num(),
+					ColorWriteMask);
+			}
+
+			LastDynamicStateUpdate->ShaderKeyShared = NewDynamicState.ShaderKeyShared;
+			LastDynamicStateUpdate->ColorAttachmentStates = NewDynamicState.ColorAttachmentStates;
+		}
+	}
+
+#undef VULKAN_DYNAMIC_STATE_COMPARE
+#undef VULKAN_DYNAMIC_STATE_UPDATE
+#undef VULKAN_DYNAMIC_STATE_BLOCK_START
+#undef VULKAN_DYNAMIC_STATE_BLOCK_END
+
+	CommandBuffer.bNeedsFullDynamicStateUpdate = false;
 }
 
 void FVulkanPendingGfxState::UpdateInputAttachments(FVulkanFramebuffer* Framebuffer)
diff --git a/Engine/Source/Runtime/VulkanRHI/Private/VulkanPendingState.h b/Engine/Source/Runtime/VulkanRHI/Private/VulkanPendingState.h
--- a/Engine/Source/Runtime/VulkanRHI/Private/VulkanPendingState.h
+++ b/Engine/Source/Runtime/VulkanRHI/Private/VulkanPendingState.h
@@ -138,6 +138,7 @@ public:
 		CurrentPipeline = nullptr;
 		CurrentState = nullptr;
 		bDirtyVertexStreams = true;
+		bDirtyDynamicStrides = true;
 
 		PrimitiveType = PT_Num;
 
@@ -328,6 +329,7 @@ protected:
 	};
 	FVertexStream PendingStreams[MaxVertexElementCount];
 	bool bDirtyVertexStreams;
+	bool bDirtyDynamicStrides;
 
 	void UpdateInputAttachments(FVulkanFramebuffer* Framebuffer);
 
diff --git a/Engine/Source/Runtime/VulkanRHI/Private/VulkanPipeline.cpp b/Engine/Source/Runtime/VulkanRHI/Private/VulkanPipeline.cpp
--- a/Engine/Source/Runtime/VulkanRHI/Private/VulkanPipeline.cpp
+++ b/Engine/Source/Runtime/VulkanRHI/Private/VulkanPipeline.cpp
@@ -1199,6 +1199,95 @@ static FString GfxShaderHashesToString(FVulkanShader* Shaders[ShaderStage::NumGr
 	return ShaderHashes;
 }
 
+
+static TArray<VkDynamicState> FillDynamicStates(const FVulkanDevice& Device, const FGfxPipelineDesc& PipelineDesc)
+{
+	TArray<VkDynamicState> DynamicStates;
+
+	const bool bHasEXTExtendedDynamicState1 = Device.GetOptionalExtensions().HasEXTExtendedDynamicState1;
+	const bool bHasEXTExtendedDynamicState2 = Device.GetOptionalExtensions().HasEXTExtendedDynamicState2;
+	const bool bHasEXTExtendedDynamicState3 = Device.GetOptionalExtensions().HasEXTExtendedDynamicState3;
+
+	// Preallocate the array
+	const int32 DynamicStateCount =
+		4 +
+		(bHasEXTExtendedDynamicState1 ? 9 : 0) +
+		(bHasEXTExtendedDynamicState2 ? 1 : 0) +
+		(bHasEXTExtendedDynamicState3 ? 6 : 0);
+	DynamicStates.Reserve(DynamicStateCount);
+
+	// Basic states that are always enabled
+	DynamicStates.Add(VK_DYNAMIC_STATE_STENCIL_REFERENCE);
+	DynamicStates.Add(VK_DYNAMIC_STATE_DEPTH_BOUNDS);
+
+	if (bHasEXTExtendedDynamicState1)
+	{
+#if PLATFORM_SUPPORTS_MESH_SHADERS
+#if VULKAN_USE_SHADERKEYS
+		if (PipelineDesc.ShaderKeys[ShaderStage::Mesh] == 0)
+#else
+		if (PipelineDesc.ShaderHashes.Stages[ShaderStage::Mesh] == 0)
+#endif
+#endif
+		{
+			DynamicStates.Add(VK_DYNAMIC_STATE_PRIMITIVE_TOPOLOGY);
+
+			if (PipelineDesc.VertexBindings.Num() > 0)
+			{
+				DynamicStates.Add(VK_DYNAMIC_STATE_VERTEX_INPUT_BINDING_STRIDE);
+			}
+		}
+
+		DynamicStates.Add(VK_DYNAMIC_STATE_CULL_MODE);
+		DynamicStates.Add(VK_DYNAMIC_STATE_STENCIL_TEST_ENABLE);
+		DynamicStates.Add(VK_DYNAMIC_STATE_STENCIL_OP);
+		DynamicStates.Add(VK_DYNAMIC_STATE_DEPTH_COMPARE_OP);
+		DynamicStates.Add(VK_DYNAMIC_STATE_DEPTH_TEST_ENABLE);
+		DynamicStates.Add(VK_DYNAMIC_STATE_DEPTH_BOUNDS_TEST_ENABLE);
+		DynamicStates.Add(VK_DYNAMIC_STATE_DEPTH_WRITE_ENABLE);
+
+		// VK_DYNAMIC_STATE_FRONT_FACE unused
+
+		DynamicStates.Add(VK_DYNAMIC_STATE_VIEWPORT_WITH_COUNT);
+		DynamicStates.Add(VK_DYNAMIC_STATE_SCISSOR_WITH_COUNT);
+	}
+	else
+	{
+		DynamicStates.Add(VK_DYNAMIC_STATE_VIEWPORT);
+		DynamicStates.Add(VK_DYNAMIC_STATE_SCISSOR);
+	}
+
+	if (bHasEXTExtendedDynamicState2)
+	{
+		DynamicStates.Add(VK_DYNAMIC_STATE_DEPTH_BIAS_ENABLE);
+
+		// VK_DYNAMIC_STATE_LOGIC_OP_EXT not used
+		// VK_DYNAMIC_STATE_PATCH_CONTROL_POINTS_EXT not used
+		// VK_DYNAMIC_STATE_PRIMITIVE_RESTART_ENABLE not used
+		// VK_DYNAMIC_STATE_RASTERIZER_DISCARD_ENABLE not used
+	}
+
+	if (bHasEXTExtendedDynamicState3)
+	{
+		if (PipelineDesc.ColorAttachmentStates.Num() > 0)
+		{
+			DynamicStates.Add(VK_DYNAMIC_STATE_COLOR_BLEND_ENABLE_EXT);
+			DynamicStates.Add(VK_DYNAMIC_STATE_COLOR_WRITE_MASK_EXT);
+			DynamicStates.Add(VK_DYNAMIC_STATE_COLOR_BLEND_EQUATION_EXT);
+		}
+		DynamicStates.Add(VK_DYNAMIC_STATE_POLYGON_MODE_EXT);
+		DynamicStates.Add(VK_DYNAMIC_STATE_RASTERIZATION_SAMPLES_EXT);
+		DynamicStates.Add(VK_DYNAMIC_STATE_ALPHA_TO_COVERAGE_ENABLE_EXT);
+
+		// VK_DYNAMIC_STATE_ALPHA_TO_ONE_ENABLE_EXT not used
+		// VK_DYNAMIC_STATE_DEPTH_CLAMP_ENABLE_EXT not used (:todo-jn: but should be?)
+		// VK_DYNAMIC_STATE_LOGIC_OP_ENABLE_EXT not used
+		// VK_DYNAMIC_STATE_SAMPLE_MASK_EXT not used
+	}
+
+	return DynamicStates;
+}
+
 bool FVulkanPipelineStateCacheManager::CreateGfxPipelineFromEntry(FVulkanGraphicsPipelineState* PSO, FVulkanShader* Shaders[ShaderStage::NumGraphicsStages], FGraphicsPipelineStateInitializer::EPSOPrecacheCompileType PSOCompileType)
 {
 	VkPipeline* Pipeline = &PSO->VulkanPipeline;
@@ -1247,8 +1336,8 @@ bool FVulkanPipelineStateCacheManager::CreateGfxPipelineFromEntry(FVulkanGraphic
 	// Viewport
 	VkPipelineViewportStateCreateInfo VPInfo;
 	ZeroVulkanStruct(VPInfo, VK_STRUCTURE_TYPE_PIPELINE_VIEWPORT_STATE_CREATE_INFO);
-	VPInfo.viewportCount = 1;
-	VPInfo.scissorCount = 1;
+	VPInfo.viewportCount = Device->GetOptionalExtensions().HasEXTExtendedDynamicState1 ? 0: 1;
+	VPInfo.scissorCount = Device->GetOptionalExtensions().HasEXTExtendedDynamicState1 ? 0 : 1;
 
 	// Multisample
 	VkPipelineMultisampleStateCreateInfo MSInfo;
@@ -1384,14 +1473,9 @@ bool FVulkanPipelineStateCacheManager::CreateGfxPipelineFromEntry(FVulkanGraphic
 
 	VkPipelineDynamicStateCreateInfo DynamicState;
 	ZeroVulkanStruct(DynamicState, VK_STRUCTURE_TYPE_PIPELINE_DYNAMIC_STATE_CREATE_INFO);
-	TArray<VkDynamicState, TInlineAllocator<4>> DynamicStatesEnabled;
-	DynamicStatesEnabled.Add(VK_DYNAMIC_STATE_VIEWPORT);
-	DynamicStatesEnabled.Add(VK_DYNAMIC_STATE_SCISSOR);
-	DynamicStatesEnabled.Add(VK_DYNAMIC_STATE_STENCIL_REFERENCE);
-	DynamicStatesEnabled.Add(VK_DYNAMIC_STATE_DEPTH_BOUNDS);
+	TArray<VkDynamicState> DynamicStatesEnabled = FillDynamicStates(*PSO->Device, *GfxEntry);
 	DynamicState.pDynamicStates = DynamicStatesEnabled.GetData();
 	DynamicState.dynamicStateCount = DynamicStatesEnabled.Num();
-
 	PipelineInfo.pDynamicState = &DynamicState;
 
 	const bool bUsingVariableRateShading = (PSO->Desc.ShadingRate != EVRSShadingRate::VRSSR_1x1) || 
