VulkanRHI: Add Vulkan support for AliasableFormats in RHITextureDesc

#rnx

diff --git a/Engine/Source/Runtime/VulkanRHI/Private/VulkanTexture.cpp b/Engine/Source/Runtime/VulkanRHI/Private/VulkanTexture.cpp
--- a/Engine/Source/Runtime/VulkanRHI/Private/VulkanTexture.cpp
+++ b/Engine/Source/Runtime/VulkanRHI/Private/VulkanTexture.cpp
@@ -289,7 +289,7 @@ void FVulkanTexture::GenerateImageCreateInfo(
 	ImageCreateInfo.flags = (ResourceType == VK_IMAGE_VIEW_TYPE_CUBE || ResourceType == VK_IMAGE_VIEW_TYPE_CUBE_ARRAY) ? VK_IMAGE_CREATE_CUBE_COMPATIBLE_BIT : 0;
 
 	const bool bHasUAVFormat = (InDesc.UAVFormat != PF_Unknown && InDesc.UAVFormat != InDesc.Format);
-	const bool bNeedsMutableFormat = (EnumHasAllFlags(UEFlags, TexCreate_SRGB) || (InDesc.Format == PF_R64_UINT) || bHasUAVFormat);
+	const bool bNeedsMutableFormat = (EnumHasAllFlags(UEFlags, TexCreate_SRGB) || (InDesc.Format == PF_R64_UINT) || bHasUAVFormat || !InDesc.AliasableFormats.IsEmpty());
 	if (bNeedsMutableFormat)
 	{
 		if (InDevice.GetOptionalExtensions().HasKHRImageFormatList)
@@ -300,7 +300,7 @@ void FVulkanTexture::GenerateImageCreateInfo(
 			ImageCreateInfo.pNext = &ImageFormatListCreateInfo;
 
 			// Allow non-SRGB views to be created for SRGB textures
-			if (EnumHasAllFlags(UEFlags, TexCreate_SRGB) && nonSrgbFormat != srgbFormat)
+			if (EnumHasAllFlags(UEFlags, TexCreate_SRGB) && (nonSrgbFormat != srgbFormat))
 			{
 				OutImageCreateInfo.FormatsUsed.Add(nonSrgbFormat);
 				OutImageCreateInfo.FormatsUsed.Add(srgbFormat);
@@ -319,6 +319,15 @@ void FVulkanTexture::GenerateImageCreateInfo(
 				OutImageCreateInfo.FormatsUsed.AddUnique(UEToVkTextureFormat(InDesc.UAVFormat, false));
 			}
 
+			for (EPixelFormat AliasPixelFormat : InDesc.AliasableFormats)
+			{
+				const VkFormat AliasFormat = UEToVkTextureFormat(AliasPixelFormat, false);
+				if (ImageCreateInfo.format != AliasFormat)
+				{
+					OutImageCreateInfo.FormatsUsed.AddUnique(AliasFormat);
+				}
+			}
+
 			ImageFormatListCreateInfo.pViewFormats = OutImageCreateInfo.FormatsUsed.GetData();
 			ImageFormatListCreateInfo.viewFormatCount = OutImageCreateInfo.FormatsUsed.Num();
 		}
