VulkanRHI: Replace the use of VK_DYNAMIC_STATE_VERTEX_INPUT_BINDING_STRIDE with VK_DYNAMIC_STATE_VERTEX_INPUT_EXT (from VK_EXT_vertex_input_dynamic_state)

#rnx

diff --git a/Engine/Source/Runtime/VulkanRHI/Private/VulkanDevice.h b/Engine/Source/Runtime/VulkanRHI/Private/VulkanDevice.h
--- a/Engine/Source/Runtime/VulkanRHI/Private/VulkanDevice.h
+++ b/Engine/Source/Runtime/VulkanRHI/Private/VulkanDevice.h
@@ -64,6 +64,7 @@ struct FOptionalVulkanDeviceExtensions
 			uint64 HasEXTMutableDescriptorType : 1;
 			uint64 HasKHRMaintenance7 : 1;
 			uint64 HasEXTShaderObject : 1;
+			uint64 HasEXTGraphicsPipelineLibrary : 1;
 
 			// Vendor specific
 			uint64 HasAMDBufferMarker : 1;
@@ -103,6 +104,7 @@ struct FOptionalVulkanDeviceExtensions
 			uint64 HasEXTExtendedDynamicState1 : 1;
 			uint64 HasEXTExtendedDynamicState2 : 1;
 			uint64 HasEXTExtendedDynamicState3 : 1;
+			uint64 HasEXTVertexInputDynamicState : 1;
 
 			// Promoted to 1.4
 			uint64 HasEXTLoadStoreOpNone : 1;
@@ -135,7 +137,10 @@ struct FOptionalVulkanDeviceExtensions
 
 	bool HasAnyExtendedDynamicState() const
 	{
-		return HasEXTExtendedDynamicState1 || HasEXTExtendedDynamicState2 || HasEXTExtendedDynamicState3;
+		return HasEXTExtendedDynamicState1 || 
+			HasEXTExtendedDynamicState2 || 
+			HasEXTExtendedDynamicState3 || 
+			HasEXTVertexInputDynamicState;
 	}
 };
 
@@ -538,16 +543,33 @@ public:
 		return OptionalDeviceExtensionProperties;
 	}
 
-	inline bool NeedsAllPlanes() const
+	bool NeedsAllPlanes() const
 	{
 		return !SupportsParallelRendering();
 	}
 
-	inline bool SupportsParallelRendering() const
+	bool SupportsParallelRendering() const
 	{
 		return OptionalDeviceExtensions.HasSeparateDepthStencilLayouts && OptionalDeviceExtensions.HasKHRSynchronization2 && OptionalDeviceExtensions.HasKHRRenderPass2;
 	}
 
+	bool SupportsGraphicPipelineLibraries() const
+	{
+		return OptionalDeviceExtensions.HasEXTGraphicsPipelineLibrary &&
+			// Extended dynamic states make it possible to create the pipeline library at shader creation,
+			// and they allow for a single library for vertex_input and fragment_output_state
+			OptionalDeviceExtensions.HasEXTExtendedDynamicState1 &&
+			OptionalDeviceExtensions.HasEXTExtendedDynamicState2 &&
+			OptionalDeviceExtensions.HasEXTExtendedDynamicState3 &&
+			// use a single vertex input stage for all pipelines
+			OptionalDeviceExtensions.HasEXTVertexInputDynamicState &&
+			// dynamic rendering removes need for renderpass at pipeline creation
+			OptionalDeviceExtensions.HasKHRDynamicRendering &&  
+			// in bindless a single pipeline layout is needed
+			SupportsBindless();
+
+	}
+
 	void SetupPresentQueue(VkSurfaceKHR Surface);
 
 	const TArray<VkQueueFamilyProperties>& GetQueueFamilyProps()
diff --git a/Engine/Source/Runtime/VulkanRHI/Private/VulkanExtensions.cpp b/Engine/Source/Runtime/VulkanRHI/Private/VulkanExtensions.cpp
--- a/Engine/Source/Runtime/VulkanRHI/Private/VulkanExtensions.cpp
+++ b/Engine/Source/Runtime/VulkanRHI/Private/VulkanExtensions.cpp
@@ -100,6 +100,15 @@ TAutoConsoleVariable<int32> GVulkanAllowShaderObjectCVar(
 	ECVF_ReadOnly
 );
 
+TAutoConsoleVariable<int32> GVulkanAllowGraphicPipelineLibraryCVar(
+	TEXT("r.Vulkan.AllowGraphicPipelineLibrary"),
+	0,
+	TEXT("Enables the use of VK_EXT_graphics_pipeline_library to replace graphic pipelines. \n")
+	TEXT("0: Do not enable support for graphic pipeline library (default)\n")
+	TEXT("1: Enable graphic pipeline library"),
+	ECVF_ReadOnly
+);
+
 TAutoConsoleVariable<int32> GVulkanVariableRateShadingFormatCVar(
 	TEXT("r.Vulkan.VRSFormat"),
 	0,
@@ -1757,14 +1766,15 @@ private:
 	VkPhysicalDeviceImageCompressionControlFeaturesEXT PhysicalDeviceImageCompressionControlFeatures;
 };
 
-// ***** VK_KHR_graphics_pipeline_library
-class FVulkanKHRGraphicsPipelineLibraryExtension : public FVulkanDeviceExtension
+// ***** VK_EXT_graphics_pipeline_library
+class FVulkanEXTGraphicsPipelineLibraryExtension : public FVulkanDeviceExtension
 {
 public:
 
-	FVulkanKHRGraphicsPipelineLibraryExtension(FVulkanDevice* InDevice)
+	FVulkanEXTGraphicsPipelineLibraryExtension(FVulkanDevice* InDevice)
 		: FVulkanDeviceExtension(InDevice, VK_EXT_GRAPHICS_PIPELINE_LIBRARY_EXTENSION_NAME, VULKAN_EXTENSION_ENABLED)
 	{
+		bEnabledInCode = bEnabledInCode && (GVulkanAllowGraphicPipelineLibraryCVar.GetValueOnAnyThread() != 0);
 	}
 
 	virtual void PrePhysicalDeviceFeatures(VkPhysicalDeviceFeatures2KHR& PhysicalDeviceFeatures2) override final
@@ -1776,6 +1786,7 @@ public:
 	virtual void PostPhysicalDeviceFeatures(FOptionalVulkanDeviceExtensions& ExtensionFlags) override final
 	{
 		bRequirementsPassed = (GraphicsPipelineLibraryFeatures.graphicsPipelineLibrary == VK_TRUE);
+		ExtensionFlags.HasEXTGraphicsPipelineLibrary = bRequirementsPassed;
 	}
 
 	virtual void PrePhysicalDeviceProperties(VkPhysicalDeviceProperties2KHR& PhysicalDeviceProperties2) override final
@@ -2045,7 +2056,8 @@ public:
 			(ExtendedDynamicState3Features.extendedDynamicState3ColorBlendEnable == VK_TRUE) &&
 			(ExtendedDynamicState3Features.extendedDynamicState3ColorBlendEquation == VK_TRUE) &&
 			(ExtendedDynamicState3Features.extendedDynamicState3ColorWriteMask == VK_TRUE) &&
-			(ExtendedDynamicState3Features.extendedDynamicState3RasterizationSamples == VK_TRUE);
+			(ExtendedDynamicState3Features.extendedDynamicState3RasterizationSamples == VK_TRUE) &&
+			(ExtendedDynamicState3Features.extendedDynamicState3AlphaToCoverageEnable == VK_TRUE);
 		ExtensionFlags.HasEXTExtendedDynamicState3 = bRequirementsPassed;
 	}
 
@@ -2060,6 +2072,40 @@ public:
 	VkPhysicalDeviceExtendedDynamicState3FeaturesEXT ExtendedDynamicState3Features;
 };
 
+// ***** VK_EXT_vertex_input_dynamic_state
+class FVulkanEXTVertexInputDynamicStateExtension : public FVulkanDeviceExtension
+{
+public:
+
+	FVulkanEXTVertexInputDynamicStateExtension(FVulkanDevice* InDevice)
+		: FVulkanDeviceExtension(InDevice, VK_EXT_VERTEX_INPUT_DYNAMIC_STATE_EXTENSION_NAME, VULKAN_EXTENSION_ENABLED)
+	{
+		bEnabledInCode = bEnabledInCode && (GVulkanAllowDynamicStatesCVar.GetValueOnAnyThread() != 0);
+	}
+
+	virtual void PrePhysicalDeviceFeatures(VkPhysicalDeviceFeatures2KHR& PhysicalDeviceFeatures2) override final
+	{
+		ZeroVulkanStruct(VertexInputDynamicStateFeatures, VK_STRUCTURE_TYPE_PHYSICAL_DEVICE_VERTEX_INPUT_DYNAMIC_STATE_FEATURES_EXT);
+		AddToPNext(PhysicalDeviceFeatures2, VertexInputDynamicStateFeatures);
+	}
+
+	virtual void PostPhysicalDeviceFeatures(FOptionalVulkanDeviceExtensions& ExtensionFlags) override final
+	{
+		bRequirementsPassed = (VertexInputDynamicStateFeatures.vertexInputDynamicState == VK_TRUE);
+		ExtensionFlags.HasEXTVertexInputDynamicState = bRequirementsPassed;
+	}
+
+	virtual void PreCreateDevice(VkDeviceCreateInfo& DeviceCreateInfo) override final
+	{
+		if (bRequirementsPassed)
+		{
+			AddToPNext(DeviceCreateInfo, VertexInputDynamicStateFeatures);
+		}
+	}
+
+	VkPhysicalDeviceVertexInputDynamicStateFeaturesEXT VertexInputDynamicStateFeatures;
+};
+
 // ***** VK_EXT_shader_object
 class FVulkanEXTShaderObjectExtension : public FVulkanDeviceExtension
 {
@@ -2191,15 +2237,15 @@ FVulkanDeviceExtensionArray FVulkanDeviceExtension::GetUESupportedDeviceExtensio
 	ADD_CUSTOM_EXTENSION(FVulkanKHRTimelineSemaphoreExtension);
 	ADD_CUSTOM_EXTENSION(FVulkanEXTMeshShaderExtension);
 	ADD_CUSTOM_EXTENSION(FVulkanEXTImageCompressionControlExtension);
-	ADD_CUSTOM_EXTENSION(FVulkanKHRGraphicsPipelineLibraryExtension);
+	ADD_CUSTOM_EXTENSION(FVulkanEXTGraphicsPipelineLibraryExtension);
 	ADD_CUSTOM_EXTENSION(FVulkanEXTMutableDescriptorTypeExtension);
 	ADD_CUSTOM_EXTENSION(FVulkanKHRDynamicRenderingExtension);			// must be kept BEFORE DynamicRenderingLocalRead!
 	ADD_CUSTOM_EXTENSION(FVulkanKHRDynamicRenderingLocalReadExtension);	// must be kept AFTER DynamicRendering!
 	ADD_CUSTOM_EXTENSION(FVulkanKHRDynamicRenderingUnusedAttachmentsExtension);	// must be kept AFTER DynamicRendering!
-	ADD_CUSTOM_EXTENSION(FVulkanKHRDynamicRenderingUnusedAttachmentsExtension);	// must be kept AFTER DynamicRendering!
 	ADD_CUSTOM_EXTENSION(FVulkanEXTExtendedDynamicStateExtension);
 	ADD_CUSTOM_EXTENSION(FVulkanEXTExtendedDynamicState2Extension);
 	ADD_CUSTOM_EXTENSION(FVulkanEXTExtendedDynamicState3Extension);
+	ADD_CUSTOM_EXTENSION(FVulkanEXTVertexInputDynamicStateExtension);
 	ADD_CUSTOM_EXTENSION(FVulkanEXTShaderObjectExtension);
 
 	// Needed for Raytracing
diff --git a/Engine/Source/Runtime/VulkanRHI/Private/VulkanPendingState.cpp b/Engine/Source/Runtime/VulkanRHI/Private/VulkanPendingState.cpp
--- a/Engine/Source/Runtime/VulkanRHI/Private/VulkanPendingState.cpp
+++ b/Engine/Source/Runtime/VulkanRHI/Private/VulkanPendingState.cpp
@@ -390,13 +390,6 @@ void FVulkanPendingGfxState::PrepareForDraw(FVulkanCommandListContext& Context)
 
 	check(Context.GetCommandBuffer().bHasPipeline);
 
-	// Force an update on vertex stream strides
-	if (Context.GetCommandBuffer().bNeedsFullDynamicStateUpdate && 
-		Device.GetOptionalExtensions().HasEXTExtendedDynamicState1)
-	{
-		bDirtyDynamicStrides = true;
-	}
-
 	if (CurrentPipeline->UsesBindless())
 	{
 		check(!CurrentPipeline->bHasInputAttachments); // todo-jn: bindless + InputAttachments
@@ -422,7 +415,7 @@ void FVulkanPendingGfxState::PrepareForDraw(FVulkanCommandListContext& Context)
 		}
 	}
 
-	if (bDirtyVertexStreams || bDirtyDynamicStrides)
+	if (bDirtyVertexStreams)
 	{
 #if VULKAN_ENABLE_AGGRESSIVE_STATS
 		SCOPE_CYCLE_COUNTER(STAT_VulkanBindVertexStreamsTime);
@@ -433,10 +426,6 @@ void FVulkanPendingGfxState::PrepareForDraw(FVulkanCommandListContext& Context)
 		{
 			// However, we need to verify that there are also no bindings
 			check(VertexInputStateInfo.BindingsNum == 0);
-
-			// Will clear the dynamic strides, be sure to set them on the next pass
-			bDirtyDynamicStrides = true;
-
 			return;
 		}
 
@@ -499,19 +488,7 @@ void FVulkanPendingGfxState::PrepareForDraw(FVulkanCommandListContext& Context)
 			// Incorrect:	1, 0, 2, 3
 			// Incorrect:	0, 2, 3, 5
 			// Reordering and creation of stream binding index is done in "GenerateVertexInputStateInfo()"
-			if (Device.GetOptionalExtensions().HasEXTExtendedDynamicState1)
-			{
-				VulkanRHI::vkCmdBindVertexBuffers2EXT(Context.GetCommandBuffer().GetHandle(), 0, 
-					TemporaryIA.NumUsed, TemporaryIA.VertexBuffers, TemporaryIA.VertexOffsets, TemporaryIA.VertexSizes, 
-					CurrentPipeline->GetVulkanShader(SF_Mesh) ? nullptr : TemporaryIA.VertexStrides);
-
-				// Mesh shader will clear the dynamic strides, be sure to set them on the next pass
-				bDirtyDynamicStrides = (CurrentPipeline->GetVulkanShader(SF_Mesh) == nullptr);
-			}
-			else
-			{
-				VulkanRHI::vkCmdBindVertexBuffers(Context.GetCommandBuffer().GetHandle(), 0, TemporaryIA.NumUsed, TemporaryIA.VertexBuffers, TemporaryIA.VertexOffsets);
-			}
+			VulkanRHI::vkCmdBindVertexBuffers(Context.GetCommandBuffer().GetHandle(), 0, TemporaryIA.NumUsed, TemporaryIA.VertexBuffers, TemporaryIA.VertexOffsets);
 		}
 
 		bDirtyVertexStreams = false;
@@ -577,6 +554,7 @@ void FVulkanPendingGfxState::UpdateDynamicStates(FVulkanCommandBuffer& CommandBu
 	const bool bForceUpdate = CommandBuffer.bNeedsFullDynamicStateUpdate;
 	FGfxPipelineDesc* LastDynamicStateUpdate = CommandBuffer.LastDynamicStateUpdate;
 	const FGfxPipelineDesc& NewDynamicState = CurrentPipeline->Desc;
+	const bool bHasMeshShader = (CurrentPipeline->GetVulkanShader(SF_Mesh) != nullptr);
 
 	// Bunch of macros to make dynamic state comparions below easier to read
 #define VULKAN_DYNAMIC_STATE_COMPARE(Name) (LastDynamicStateUpdate->Name != NewDynamicState.Name)
@@ -598,7 +576,7 @@ void FVulkanPendingGfxState::UpdateDynamicStates(FVulkanCommandBuffer& CommandBu
 	{
 		check(LastDynamicStateUpdate);
 
-		if (CurrentPipeline->GetVulkanShader(SF_Mesh) == nullptr)
+		if (!bHasMeshShader)
 		{
 			// VK_DYNAMIC_STATE_PRIMITIVE_TOPOLOGY
 			VULKAN_DYNAMIC_STATE_BLOCK_START(Topology);
@@ -760,6 +738,43 @@ void FVulkanPendingGfxState::UpdateDynamicStates(FVulkanCommandBuffer& CommandBu
 		}
 	}
 
+	if (Device.GetOptionalExtensions().HasEXTVertexInputDynamicState)
+	{
+		bDirtyDynamicVertexInput = bDirtyDynamicVertexInput || bForceUpdate;
+
+		if (bDirtyDynamicVertexInput ||
+			(LastDynamicStateUpdate->VertexAttributes != NewDynamicState.VertexAttributes) ||
+			(LastDynamicStateUpdate->VertexBindings != NewDynamicState.VertexBindings))
+		{
+			if (!bHasMeshShader)
+			{
+				TArray<VkVertexInputBindingDescription2EXT, TInlineAllocator<32>> BindingDescription;
+				for (const FGfxPipelineDesc::FVertexBinding& SourceBinding : NewDynamicState.VertexBindings)
+				{
+					VkVertexInputBindingDescription2EXT& Binding = BindingDescription.AddZeroed_GetRef();
+					Binding.sType = VK_STRUCTURE_TYPE_VERTEX_INPUT_BINDING_DESCRIPTION_2_EXT;
+					SourceBinding.WriteInto(Binding);
+				}
+
+				TArray<VkVertexInputAttributeDescription2EXT, TInlineAllocator<32>> AttributeDescription;
+				for (const FGfxPipelineDesc::FVertexAttribute& SourceAttr : NewDynamicState.VertexAttributes)
+				{
+					VkVertexInputAttributeDescription2EXT& Attr = AttributeDescription.AddZeroed_GetRef();
+					Attr.sType = VK_STRUCTURE_TYPE_VERTEX_INPUT_ATTRIBUTE_DESCRIPTION_2_EXT;
+					SourceAttr.WriteInto(Attr);
+				}
+
+				VulkanRHI::vkCmdSetVertexInputEXT(CommandBuffer.GetHandle(),
+					BindingDescription.Num(), BindingDescription.GetData(),
+					AttributeDescription.Num(), AttributeDescription.GetData());
+
+				LastDynamicStateUpdate->VertexAttributes = NewDynamicState.VertexAttributes;
+				LastDynamicStateUpdate->VertexBindings = NewDynamicState.VertexBindings;
+				bDirtyDynamicVertexInput = false;
+			}
+		}
+	}
+
 #undef VULKAN_DYNAMIC_STATE_COMPARE
 #undef VULKAN_DYNAMIC_STATE_UPDATE
 #undef VULKAN_DYNAMIC_STATE_BLOCK_START
diff --git a/Engine/Source/Runtime/VulkanRHI/Private/VulkanPendingState.h b/Engine/Source/Runtime/VulkanRHI/Private/VulkanPendingState.h
--- a/Engine/Source/Runtime/VulkanRHI/Private/VulkanPendingState.h
+++ b/Engine/Source/Runtime/VulkanRHI/Private/VulkanPendingState.h
@@ -138,7 +138,7 @@ public:
 		CurrentPipeline = nullptr;
 		CurrentState = nullptr;
 		bDirtyVertexStreams = true;
-		bDirtyDynamicStrides = true;
+		bDirtyDynamicVertexInput = true;
 
 		PrimitiveType = PT_Num;
 
@@ -329,7 +329,7 @@ protected:
 	};
 	FVertexStream PendingStreams[MaxVertexElementCount];
 	bool bDirtyVertexStreams;
-	bool bDirtyDynamicStrides;
+	bool bDirtyDynamicVertexInput;
 
 	void UpdateInputAttachments(FVulkanFramebuffer* Framebuffer);
 
diff --git a/Engine/Source/Runtime/VulkanRHI/Private/VulkanPipeline.cpp b/Engine/Source/Runtime/VulkanRHI/Private/VulkanPipeline.cpp
--- a/Engine/Source/Runtime/VulkanRHI/Private/VulkanPipeline.cpp
+++ b/Engine/Source/Runtime/VulkanRHI/Private/VulkanPipeline.cpp
@@ -755,6 +755,14 @@ void FGfxPipelineDesc::FVertexBinding::WriteInto(VkVertexInputBindingDescription
 	Out.stride =	Stride;
 }
 
+void FGfxPipelineDesc::FVertexBinding::WriteInto(VkVertexInputBindingDescription2EXT& Out) const
+{
+	Out.binding =	Binding;
+	Out.inputRate =	(VkVertexInputRate)InputRate;
+	Out.stride =	Stride;
+	Out.divisor =	1;
+}
+
 FArchive& operator << (FArchive& Ar, FGfxPipelineDesc::FVertexBinding& Binding)
 {
 	// Modify VERSION if serialization changes
@@ -780,6 +788,14 @@ void FGfxPipelineDesc::FVertexAttribute::WriteInto(VkVertexInputAttributeDescrip
 	Out.offset =	Offset;
 }
 
+void FGfxPipelineDesc::FVertexAttribute::WriteInto(VkVertexInputAttributeDescription2EXT& Out) const
+{
+	Out.binding = Binding;
+	Out.format = (VkFormat)Format;
+	Out.location = Location;
+	Out.offset = Offset;
+}
+
 FArchive& operator << (FArchive& Ar, FGfxPipelineDesc::FVertexAttribute& Attribute)
 {
 	// Modify VERSION if serialization changes
@@ -1207,13 +1223,25 @@ static TArray<VkDynamicState> FillDynamicStates(const FVulkanDevice& Device, con
 	const bool bHasEXTExtendedDynamicState1 = Device.GetOptionalExtensions().HasEXTExtendedDynamicState1;
 	const bool bHasEXTExtendedDynamicState2 = Device.GetOptionalExtensions().HasEXTExtendedDynamicState2;
 	const bool bHasEXTExtendedDynamicState3 = Device.GetOptionalExtensions().HasEXTExtendedDynamicState3;
+	const bool bHasEXTVertexInputDynamicState = Device.GetOptionalExtensions().HasEXTVertexInputDynamicState;
+
+#if PLATFORM_SUPPORTS_MESH_SHADERS
+#if VULKAN_USE_SHADERKEYS
+	const bool bHasMeshShader = (PipelineDesc.ShaderKeys[ShaderStage::Mesh] != 0);
+#else
+	const bool bHasMeshShader = (PipelineDesc.ShaderHashes.Stages[ShaderStage::Mesh] != 0);
+#endif
+#else
+	constexpr bool bHasMeshShader = false;
+#endif
 
 	// Preallocate the array
 	const int32 DynamicStateCount =
 		4 +
-		(bHasEXTExtendedDynamicState1 ? 9 : 0) +
+		(bHasEXTExtendedDynamicState1 ? 8 : 0) +
 		(bHasEXTExtendedDynamicState2 ? 1 : 0) +
-		(bHasEXTExtendedDynamicState3 ? 6 : 0);
+		(bHasEXTExtendedDynamicState3 ? 6 : 0) +
+		(bHasEXTVertexInputDynamicState ? 1 : 0);
 	DynamicStates.Reserve(DynamicStateCount);
 
 	// Basic states that are always enabled
@@ -1222,20 +1250,9 @@ static TArray<VkDynamicState> FillDynamicStates(const FVulkanDevice& Device, con
 
 	if (bHasEXTExtendedDynamicState1)
 	{
-#if PLATFORM_SUPPORTS_MESH_SHADERS
-#if VULKAN_USE_SHADERKEYS
-		if (PipelineDesc.ShaderKeys[ShaderStage::Mesh] == 0)
-#else
-		if (PipelineDesc.ShaderHashes.Stages[ShaderStage::Mesh] == 0)
-#endif
-#endif
+		if (!bHasMeshShader)
 		{
 			DynamicStates.Add(VK_DYNAMIC_STATE_PRIMITIVE_TOPOLOGY);
-
-			if (PipelineDesc.VertexBindings.Num() > 0)
-			{
-				DynamicStates.Add(VK_DYNAMIC_STATE_VERTEX_INPUT_BINDING_STRIDE);
-			}
 		}
 
 		DynamicStates.Add(VK_DYNAMIC_STATE_CULL_MODE);
@@ -1247,6 +1264,7 @@ static TArray<VkDynamicState> FillDynamicStates(const FVulkanDevice& Device, con
 		DynamicStates.Add(VK_DYNAMIC_STATE_DEPTH_WRITE_ENABLE);
 
 		// VK_DYNAMIC_STATE_FRONT_FACE unused
+		// VK_DYNAMIC_STATE_VERTEX_INPUT_BINDING_STRIDE not used, using VK_DYNAMIC_STATE_VERTEX_INPUT_EXT instead
 
 		DynamicStates.Add(VK_DYNAMIC_STATE_VIEWPORT_WITH_COUNT);
 		DynamicStates.Add(VK_DYNAMIC_STATE_SCISSOR_WITH_COUNT);
@@ -1285,6 +1303,14 @@ static TArray<VkDynamicState> FillDynamicStates(const FVulkanDevice& Device, con
 		// VK_DYNAMIC_STATE_SAMPLE_MASK_EXT not used
 	}
 
+	if (bHasEXTVertexInputDynamicState)
+	{
+		if (!bHasMeshShader)
+		{
+			DynamicStates.Add(VK_DYNAMIC_STATE_VERTEX_INPUT_EXT);
+		}
+	}
+
 	return DynamicStates;
 }
 
@@ -1298,7 +1324,6 @@ bool FVulkanPipelineStateCacheManager::CreateGfxPipelineFromEntry(FVulkanGraphic
 	}
 
 	TRefCountPtr<FVulkanShaderModule> ShaderModules[ShaderStage::NumGraphicsStages];
-
 	PSO->GetOrCreateShaderModules(ShaderModules, Shaders);
 
 	// Pipeline
diff --git a/Engine/Source/Runtime/VulkanRHI/Private/VulkanPipeline.h b/Engine/Source/Runtime/VulkanRHI/Private/VulkanPipeline.h
--- a/Engine/Source/Runtime/VulkanRHI/Private/VulkanPipeline.h
+++ b/Engine/Source/Runtime/VulkanRHI/Private/VulkanPipeline.h
@@ -168,6 +168,7 @@ struct FGfxPipelineDesc
 
 		void ReadFrom(const VkVertexInputBindingDescription& InState);
 		void WriteInto(VkVertexInputBindingDescription& OutState) const;
+		void WriteInto(VkVertexInputBindingDescription2EXT& OutState) const;
 
 		bool operator==(const FVertexBinding& In) const
 		{
@@ -186,6 +187,7 @@ struct FGfxPipelineDesc
 
 		void ReadFrom(const VkVertexInputAttributeDescription& InState);
 		void WriteInto(VkVertexInputAttributeDescription& OutState) const;
+		void WriteInto(VkVertexInputAttributeDescription2EXT& OutState) const;
 
 		bool operator==(const FVertexAttribute& In) const
 		{
@@ -594,7 +596,7 @@ private:
 	FCriticalSection GraphicsPSOLockedCS;
 	TMap<FVulkanPSOKey, FVulkanGraphicsPipelineState*> GraphicsPSOLockedMap;
 
-		
+	
 	FCriticalSection LRUCS;
 	FVulkanGraphicsPipelineStateLRU LRU;
 	uint32 LRUUsedPipelineSize = 0;
