VulkanRHI: Clean up: Remove unused stats.  Reconnect the barrier time stat.

#rnx

diff --git a/Engine/Source/Runtime/VulkanRHI/Private/VulkanBarriers.cpp b/Engine/Source/Runtime/VulkanRHI/Private/VulkanBarriers.cpp
--- a/Engine/Source/Runtime/VulkanRHI/Private/VulkanBarriers.cpp
+++ b/Engine/Source/Runtime/VulkanRHI/Private/VulkanBarriers.cpp
@@ -1596,6 +1596,10 @@ void FVulkanCommandListContext::RHIBeginTransitions(TArrayView<const FRHITransit
 {
 	if (Device.SupportsParallelRendering())
 	{
+#if VULKAN_ENABLE_AGGRESSIVE_STATS
+		SCOPE_CYCLE_COUNTER(STAT_VulkanBarrierTime);
+#endif
+
 		constexpr bool bIsBeginTransition = true;
 
 		ProcessTransitionSync2(*this, Transitions, bIsBeginTransition);
@@ -1612,6 +1616,10 @@ void FVulkanCommandListContext::RHIBeginTransitions(TArrayView<const FRHITransit
 
 void FVulkanCommandListContext::RHIEndTransitions(TArrayView<const FRHITransition*> Transitions)
 {
+#if VULKAN_ENABLE_AGGRESSIVE_STATS
+	SCOPE_CYCLE_COUNTER(STAT_VulkanBarrierTime);
+#endif
+
 	if (Device.SupportsParallelRendering())
 	{
 		constexpr bool bIsBeginTransition = false;
diff --git a/Engine/Source/Runtime/VulkanRHI/Private/VulkanRHIPrivate.h b/Engine/Source/Runtime/VulkanRHI/Private/VulkanRHIPrivate.h
--- a/Engine/Source/Runtime/VulkanRHI/Private/VulkanRHIPrivate.h
+++ b/Engine/Source/Runtime/VulkanRHI/Private/VulkanRHIPrivate.h
@@ -159,7 +159,6 @@ DECLARE_CYCLE_STAT_EXTERN(TEXT("PSO Vulkan Creation time"), STAT_VulkanPSOVulkan
 
 // Stats
 DECLARE_STATS_GROUP(TEXT("Vulkan RHI"), STATGROUP_VulkanRHI, STATCAT_Advanced);
-//DECLARE_STATS_GROUP(TEXT("Vulkan RHI Verbose"), STATGROUP_VulkanRHIVERBOSE, STATCAT_Advanced);
 DECLARE_CYCLE_STAT_EXTERN(TEXT("Draw call time"), STAT_VulkanDrawCallTime, STATGROUP_VulkanRHI, );
 DECLARE_CYCLE_STAT_EXTERN(TEXT("Dispatch call time"), STAT_VulkanDispatchCallTime, STATGROUP_VulkanRHI, );
 DECLARE_CYCLE_STAT_EXTERN(TEXT("Draw call prep time"), STAT_VulkanDrawCallPrepareTime, STATGROUP_VulkanRHI, );
@@ -174,15 +173,8 @@ DECLARE_DWORD_ACCUMULATOR_STAT_EXTERN(TEXT("Num Frame Buffers"), STAT_VulkanNumF
 DECLARE_DWORD_ACCUMULATOR_STAT_EXTERN(TEXT("Num Buffer Views"), STAT_VulkanNumBufferViews, STATGROUP_VulkanRHI, );
 DECLARE_DWORD_ACCUMULATOR_STAT_EXTERN(TEXT("Num Image Views"), STAT_VulkanNumImageViews, STATGROUP_VulkanRHI, );
 DECLARE_DWORD_ACCUMULATOR_STAT_EXTERN(TEXT("Num Physical Mem Allocations"), STAT_VulkanNumPhysicalMemAllocations, STATGROUP_VulkanRHI, );
-DECLARE_DWORD_COUNTER_STAT_EXTERN(TEXT("Dynamic VB Size"), STAT_VulkanDynamicVBSize, STATGROUP_VulkanRHI, );
-DECLARE_DWORD_COUNTER_STAT_EXTERN(TEXT("Dynamic IB Size"), STAT_VulkanDynamicIBSize, STATGROUP_VulkanRHI, );
-DECLARE_CYCLE_STAT_EXTERN(TEXT("Dynamic VB Lock/Unlock time"), STAT_VulkanDynamicVBLockTime, STATGROUP_VulkanRHI, );
-DECLARE_CYCLE_STAT_EXTERN(TEXT("Dynamic IB Lock/Unlock time"), STAT_VulkanDynamicIBLockTime, STATGROUP_VulkanRHI, );
-DECLARE_CYCLE_STAT_EXTERN(TEXT("DrawPrim UP Prep Time"), STAT_VulkanUPPrepTime, STATGROUP_VulkanRHI, );
 DECLARE_CYCLE_STAT_EXTERN(TEXT("Uniform Buffer Creation Time"), STAT_VulkanUniformBufferCreateTime, STATGROUP_VulkanRHI, );
-DECLARE_CYCLE_STAT_EXTERN(TEXT("Apply DS Uniform Buffers"), STAT_VulkanApplyDSUniformBuffers, STATGROUP_VulkanRHI, );
 DECLARE_CYCLE_STAT_EXTERN(TEXT("Apply Packed Uniform Buffers"), STAT_VulkanApplyPackedUniformBuffers, STATGROUP_VulkanRHI, );
-DECLARE_CYCLE_STAT_EXTERN(TEXT("Barrier Time"), STAT_VulkanBarrierTime, STATGROUP_VulkanRHI, );
 DECLARE_CYCLE_STAT_EXTERN(TEXT("SRV Update Time"), STAT_VulkanSRVUpdateTime, STATGROUP_VulkanRHI, );
 DECLARE_CYCLE_STAT_EXTERN(TEXT("UAV Update Time"), STAT_VulkanUAVUpdateTime, STATGROUP_VulkanRHI, );
 DECLARE_CYCLE_STAT_EXTERN(TEXT("Deletion Queue"), STAT_VulkanDeletionQueue, STATGROUP_VulkanRHI, );
@@ -190,9 +182,7 @@ DECLARE_CYCLE_STAT_EXTERN(TEXT("Queue Submit"), STAT_VulkanQueueSubmit, STATGROU
 DECLARE_CYCLE_STAT_EXTERN(TEXT("Queue Present"), STAT_VulkanQueuePresent, STATGROUP_VulkanRHI, );
 DECLARE_DWORD_ACCUMULATOR_STAT_EXTERN(TEXT("Num Queries"), STAT_VulkanNumQueries, STATGROUP_VulkanRHI, );
 DECLARE_DWORD_ACCUMULATOR_STAT_EXTERN(TEXT("Num Query Pools"), STAT_VulkanNumQueryPools, STATGROUP_VulkanRHI, );
-DECLARE_CYCLE_STAT_EXTERN(TEXT("Wait For Query"), STAT_VulkanWaitQuery, STATGROUP_VulkanRHI, );
 DECLARE_CYCLE_STAT_EXTERN(TEXT("Wait For Fence"), STAT_VulkanWaitFence, STATGROUP_VulkanRHI, );
-DECLARE_CYCLE_STAT_EXTERN(TEXT("Reset Queries"), STAT_VulkanResetQuery, STATGROUP_VulkanRHI, );
 DECLARE_CYCLE_STAT_EXTERN(TEXT("Wait For Swapchain"), STAT_VulkanWaitSwapchain, STATGROUP_VulkanRHI, );
 DECLARE_CYCLE_STAT_EXTERN(TEXT("Acquire Backbuffer"), STAT_VulkanAcquireBackBuffer, STATGROUP_VulkanRHI, );
 DECLARE_CYCLE_STAT_EXTERN(TEXT("Staging Buffer Mgmt"), STAT_VulkanStagingBuffer, STATGROUP_VulkanRHI, );
@@ -201,6 +191,7 @@ DECLARE_DWORD_ACCUMULATOR_STAT_EXTERN(TEXT("Num Created DescSet Pools"), STAT_Vu
 DECLARE_CYCLE_STAT_EXTERN(TEXT("UpdateUniformBuffers"), STAT_VulkanUpdateUniformBuffers, STATGROUP_VulkanRHI, );
 DECLARE_CYCLE_STAT_EXTERN(TEXT("UpdateUniformBuffersRename"), STAT_VulkanUpdateUniformBuffersRename, STATGROUP_VulkanRHI, );
 #if VULKAN_ENABLE_AGGRESSIVE_STATS
+DECLARE_CYCLE_STAT_EXTERN(TEXT("Barrier Time"), STAT_VulkanBarrierTime, STATGROUP_VulkanRHI, );
 DECLARE_CYCLE_STAT_EXTERN(TEXT("Update DescriptorSets"), STAT_VulkanUpdateDescriptorSets, STATGROUP_VulkanRHI, );
 DECLARE_DWORD_COUNTER_STAT_EXTERN(TEXT("Num Desc Sets Updated"), STAT_VulkanNumDescSets, STATGROUP_VulkanRHI, );
 DECLARE_DWORD_COUNTER_STAT_EXTERN(TEXT("Num WriteDescriptors Cmd"), STAT_VulkanNumUpdateDescriptors, STATGROUP_VulkanRHI, );
diff --git a/Engine/Source/Runtime/VulkanRHI/Private/VulkanUtil.cpp b/Engine/Source/Runtime/VulkanRHI/Private/VulkanUtil.cpp
--- a/Engine/Source/Runtime/VulkanRHI/Private/VulkanUtil.cpp
+++ b/Engine/Source/Runtime/VulkanRHI/Private/VulkanUtil.cpp
@@ -1099,15 +1099,8 @@ DEFINE_STAT(STAT_VulkanNumFrameBuffers);
 DEFINE_STAT(STAT_VulkanNumBufferViews);
 DEFINE_STAT(STAT_VulkanNumImageViews);
 DEFINE_STAT(STAT_VulkanNumPhysicalMemAllocations);
-DEFINE_STAT(STAT_VulkanDynamicVBSize);
-DEFINE_STAT(STAT_VulkanDynamicIBSize);
-DEFINE_STAT(STAT_VulkanDynamicVBLockTime);
-DEFINE_STAT(STAT_VulkanDynamicIBLockTime);
-DEFINE_STAT(STAT_VulkanUPPrepTime);
 DEFINE_STAT(STAT_VulkanUniformBufferCreateTime);
-DEFINE_STAT(STAT_VulkanApplyDSUniformBuffers);
 DEFINE_STAT(STAT_VulkanApplyPackedUniformBuffers);
-DEFINE_STAT(STAT_VulkanBarrierTime);
 DEFINE_STAT(STAT_VulkanSRVUpdateTime);
 DEFINE_STAT(STAT_VulkanUAVUpdateTime);
 DEFINE_STAT(STAT_VulkanDeletionQueue);
@@ -1115,9 +1108,7 @@ DEFINE_STAT(STAT_VulkanQueueSubmit);
 DEFINE_STAT(STAT_VulkanQueuePresent);
 DEFINE_STAT(STAT_VulkanNumQueries);
 DEFINE_STAT(STAT_VulkanNumQueryPools);
-DEFINE_STAT(STAT_VulkanWaitQuery);
 DEFINE_STAT(STAT_VulkanWaitFence);
-DEFINE_STAT(STAT_VulkanResetQuery);
 DEFINE_STAT(STAT_VulkanWaitSwapchain);
 DEFINE_STAT(STAT_VulkanAcquireBackBuffer);
 DEFINE_STAT(STAT_VulkanStagingBuffer);
@@ -1126,6 +1117,7 @@ DEFINE_STAT(STAT_VulkanNumDescPools);
 DEFINE_STAT(STAT_VulkanUpdateUniformBuffers);
 DEFINE_STAT(STAT_VulkanUpdateUniformBuffersRename);
 #if VULKAN_ENABLE_AGGRESSIVE_STATS
+DEFINE_STAT(STAT_VulkanBarrierTime);
 DEFINE_STAT(STAT_VulkanUpdateDescriptorSets);
 DEFINE_STAT(STAT_VulkanNumUpdateDescriptors);
 DEFINE_STAT(STAT_VulkanNumDescSets);
