VulkanRHI: Reset vertex input when pipeline that don't have it dynamic are used (fix validation issue).

#rnx

diff --git a/Engine/Source/Runtime/VulkanRHI/Private/VulkanPendingState.cpp b/Engine/Source/Runtime/VulkanRHI/Private/VulkanPendingState.cpp
--- a/Engine/Source/Runtime/VulkanRHI/Private/VulkanPendingState.cpp
+++ b/Engine/Source/Runtime/VulkanRHI/Private/VulkanPendingState.cpp
@@ -555,6 +555,7 @@ void FVulkanPendingGfxState::UpdateDynamicStates(FVulkanCommandBuffer& CommandBu
 	FGfxPipelineDesc* LastDynamicStateUpdate = CommandBuffer.LastDynamicStateUpdate;
 	const FGfxPipelineDesc& NewDynamicState = CurrentPipeline->Desc;
 	const bool bHasMeshShader = (CurrentPipeline->GetVulkanShader(SF_Mesh) != nullptr);
+	const bool bHasGeometryShader = (CurrentPipeline->GetVulkanShader(SF_Geometry) != nullptr);
 
 	// Bunch of macros to make dynamic state comparions below easier to read
 #define VULKAN_DYNAMIC_STATE_COMPARE(Name) (LastDynamicStateUpdate->Name != NewDynamicState.Name)
@@ -773,6 +774,12 @@ void FVulkanPendingGfxState::UpdateDynamicStates(FVulkanCommandBuffer& CommandBu
 				bDirtyDynamicVertexInput = false;
 			}
 		}
+
+		// Pipelines with mesh or geometry shaders invalidate the dynamic state
+		if (bHasMeshShader || bHasGeometryShader)
+		{
+			bDirtyDynamicVertexInput = true;
+		}
 	}
 
 #undef VULKAN_DYNAMIC_STATE_COMPARE
