VulkanRHI: Remove more uses of VulkanRHI::DeviceChild (FVulkanBuffer, FVulkanTransientHeap, FVulkanTransientHeapCache, FVulkanTransientResourceAllocator, FDeferredDeletionQueue2)

#rnx

diff --git a/Engine/Source/Runtime/VulkanRHI/Private/VulkanBuffer.cpp b/Engine/Source/Runtime/VulkanRHI/Private/VulkanBuffer.cpp
--- a/Engine/Source/Runtime/VulkanRHI/Private/VulkanBuffer.cpp
+++ b/Engine/Source/Runtime/VulkanRHI/Private/VulkanBuffer.cpp
@@ -61,19 +61,19 @@ static void UpdateVulkanBufferStats(const FRHIBufferDesc& BufferDesc, int64 Buff
 	UE::RHICore::UpdateGlobalBufferStats(BufferDesc, BufferSize, bAllocating);
 }
 
-static VkDeviceAddress GetBufferDeviceAddress(FVulkanDevice* Device, VkBuffer Buffer)
+static VkDeviceAddress GetBufferDeviceAddress(FVulkanDevice& Device, VkBuffer Buffer)
 {
-	if (Device->GetOptionalExtensions().HasBufferDeviceAddress)
+	if (Device.GetOptionalExtensions().HasBufferDeviceAddress)
 	{
 		VkBufferDeviceAddressInfoKHR DeviceAddressInfo;
 		ZeroVulkanStruct(DeviceAddressInfo, VK_STRUCTURE_TYPE_BUFFER_DEVICE_ADDRESS_INFO);
 		DeviceAddressInfo.buffer = Buffer;
-		return VulkanRHI::vkGetBufferDeviceAddressKHR(Device->GetHandle(), &DeviceAddressInfo);
+		return VulkanRHI::vkGetBufferDeviceAddressKHR(Device.GetHandle(), &DeviceAddressInfo);
 	}
 	return 0;
 }
 
-VkBufferUsageFlags FVulkanBuffer::UEToVKBufferUsageFlags(FVulkanDevice* InDevice, EBufferUsageFlags InUEUsage, bool bZeroSize)
+VkBufferUsageFlags FVulkanBuffer::UEToVKBufferUsageFlags(FVulkanDevice& InDevice, EBufferUsageFlags InUEUsage, bool bZeroSize)
 {
 	// Always include TRANSFER_SRC since hardware vendors confirmed it wouldn't have any performance cost and we need it for some debug functionalities.
 	VkBufferUsageFlags OutVkUsage = VK_BUFFER_USAGE_TRANSFER_SRC_BIT;
@@ -99,7 +99,7 @@ VkBufferUsageFlags FVulkanBuffer::UEToVKBufferUsageFlags(FVulkanDevice* InDevice
 
 		TranslateFlag(BUF_Volatile, 0, VK_BUFFER_USAGE_TRANSFER_DST_BIT);
 
-		if (InDevice->GetOptionalExtensions().HasRaytracingExtensions())
+		if (InDevice.GetOptionalExtensions().HasRaytracingExtensions())
 		{
 			OutVkUsage |= VK_BUFFER_USAGE_SHADER_DEVICE_ADDRESS_BIT;
 
@@ -107,7 +107,7 @@ VkBufferUsageFlags FVulkanBuffer::UEToVKBufferUsageFlags(FVulkanDevice* InDevice
 		}
 
 		// For descriptors buffers
-		if (InDevice->GetOptionalExtensions().HasBufferDeviceAddress)
+		if (InDevice.GetOptionalExtensions().HasBufferDeviceAddress)
 		{
 			OutVkUsage |= VK_BUFFER_USAGE_SHADER_DEVICE_ADDRESS_BIT;
 		}
@@ -116,9 +116,9 @@ VkBufferUsageFlags FVulkanBuffer::UEToVKBufferUsageFlags(FVulkanDevice* InDevice
 	return OutVkUsage;
 }
 
-FVulkanBuffer::FVulkanBuffer(FVulkanDevice* InDevice, const FRHIBufferCreateDesc& CreateDesc, const FRHITransientHeapAllocation* InTransientHeapAllocation)
+FVulkanBuffer::FVulkanBuffer(FVulkanDevice& InDevice, const FRHIBufferCreateDesc& CreateDesc, const FRHITransientHeapAllocation* InTransientHeapAllocation)
 	: FRHIBuffer(CreateDesc)
-	, VulkanRHI::FDeviceChild(InDevice)
+	, Device(InDevice)
 {
 	VULKAN_TRACK_OBJECT_CREATE(FVulkanBuffer, this);
 
@@ -127,15 +127,13 @@ FVulkanBuffer::FVulkanBuffer(FVulkanDevice* InDevice, const FRHIBufferCreateDesc
 
 	if (!bZeroSize)
 	{
-		check(InDevice);
-
-		const bool bUnifiedMem = InDevice->HasUnifiedMemory();
-		const uint32 BufferAlignment = VulkanRHI::FMemoryManager::CalculateBufferAlignment(*InDevice, CreateDesc.Usage, bZeroSize);
+		const bool bUnifiedMem = InDevice.HasUnifiedMemory();
+		const uint32 BufferAlignment = VulkanRHI::FMemoryManager::CalculateBufferAlignment(InDevice, CreateDesc.Usage, bZeroSize);
 
 		if (InTransientHeapAllocation != nullptr)
 		{
 			CurrentBufferAlloc.Alloc = FVulkanTransientHeap::GetVulkanAllocation(*InTransientHeapAllocation);
-			CurrentBufferAlloc.HostPtr = bUnifiedMem ? CurrentBufferAlloc.Alloc.GetMappedPointer(Device) : nullptr;
+			CurrentBufferAlloc.HostPtr = bUnifiedMem ? CurrentBufferAlloc.Alloc.GetMappedPointer(&Device) : nullptr;
 			CurrentBufferAlloc.DeviceAddress = GetBufferDeviceAddress(InDevice, CurrentBufferAlloc.Alloc.GetBufferHandle()) + CurrentBufferAlloc.Alloc.Offset;
 			check(CurrentBufferAlloc.Alloc.Offset % BufferAlignment == 0);
 			check(CurrentBufferAlloc.Alloc.Size >= CreateDesc.Size);
@@ -157,7 +155,7 @@ FVulkanBuffer::~FVulkanBuffer()
 void FVulkanBuffer::AllocateMemory(FBufferAlloc& OutAlloc)
 {
 	VkMemoryPropertyFlags BufferMemFlags = 0;
-	const bool bUnifiedMem = Device->HasUnifiedMemory();
+	const bool bUnifiedMem = Device.HasUnifiedMemory();
 	const bool bDynamic = EnumHasAnyFlags(GetUsage(), BUF_Dynamic) || EnumHasAnyFlags(GetUsage(), BUF_Volatile);
 	if (bUnifiedMem)
 	{
@@ -172,14 +170,14 @@ void FVulkanBuffer::AllocateMemory(FBufferAlloc& OutAlloc)
 		BufferMemFlags = VK_MEMORY_PROPERTY_DEVICE_LOCAL_BIT;
 	}
 	const uint32 BufferSize = Align(GetSize(), 4);  // keep allocated size a multiple of 4 (for use with vkCmdFillBuffer)
-	const uint32 BufferAlignment = VulkanRHI::FMemoryManager::CalculateBufferAlignment(*Device, GetUsage(), (BufferSize == 0));
+	const uint32 BufferAlignment = VulkanRHI::FMemoryManager::CalculateBufferAlignment(Device, GetUsage(), (BufferSize == 0));
 
 	FBufferAlloc& NewBufferAlloc = OutAlloc;
-	if (!Device->GetMemoryManager().AllocateBufferPooled(NewBufferAlloc.Alloc, nullptr, BufferSize, BufferAlignment, BufferUsageFlags, BufferMemFlags, VulkanRHI::EVulkanAllocationMetaMultiBuffer, __FILE__, __LINE__))
+	if (!Device.GetMemoryManager().AllocateBufferPooled(NewBufferAlloc.Alloc, nullptr, BufferSize, BufferAlignment, BufferUsageFlags, BufferMemFlags, VulkanRHI::EVulkanAllocationMetaMultiBuffer, __FILE__, __LINE__))
 	{
-		Device->GetMemoryManager().HandleOOM();
+		Device.GetMemoryManager().HandleOOM();
 	}
-	NewBufferAlloc.HostPtr = (bUnifiedMem || bDynamic) ? NewBufferAlloc.Alloc.GetMappedPointer(Device) : nullptr;
+	NewBufferAlloc.HostPtr = (bUnifiedMem || bDynamic) ? NewBufferAlloc.Alloc.GetMappedPointer(&Device) : nullptr;
 	NewBufferAlloc.DeviceAddress = GetBufferDeviceAddress(Device, NewBufferAlloc.Alloc.GetBufferHandle()) + NewBufferAlloc.Alloc.Offset;
 
 	UpdateVulkanBufferStats(GetDesc(), NewBufferAlloc.Alloc.Size, true);
@@ -200,7 +198,7 @@ void* FVulkanBuffer::Lock(FRHICommandListBase& RHICmdList, EResourceLockMode Loc
 	// Static:     A single Device_Local buffer is allocated at creation.  For Lock/Unlock, use a staging buffer for the upload:
 	//             host writes to staging buffer on lock, a copy on GPU is issued on unlock to update the device_local memory.
 
-	const bool bUnifiedMem = Device->HasUnifiedMemory();
+	const bool bUnifiedMem = Device.HasUnifiedMemory();
 	const bool bDynamic = EnumHasAnyFlags(GetUsage(), BUF_Dynamic) || EnumHasAnyFlags(GetUsage(), BUF_Volatile);
 	const bool bStatic = EnumHasAnyFlags(GetUsage(), BUF_Static) || !bDynamic;
 	const bool bUAV = EnumHasAnyFlags(GetUsage(), BUF_UnorderedAccess);
@@ -221,7 +219,7 @@ void* FVulkanBuffer::Lock(FRHICommandListBase& RHICmdList, EResourceLockMode Loc
 		else 
 		{
 			// Create a staging buffer we can use to copy data from device to cpu.
-			VulkanRHI::FStagingBuffer* StagingBuffer = Device->GetStagingManager().AcquireBuffer(LockSize, VK_BUFFER_USAGE_TRANSFER_DST_BIT, VK_MEMORY_PROPERTY_HOST_CACHED_BIT);
+			VulkanRHI::FStagingBuffer* StagingBuffer = Device.GetStagingManager().AcquireBuffer(LockSize, VK_BUFFER_USAGE_TRANSFER_DST_BIT, VK_MEMORY_PROPERTY_HOST_CACHED_BIT);
 
 			VkBufferCopy Regions;
 			Regions.size = LockSize;
@@ -288,7 +286,7 @@ void* FVulkanBuffer::Lock(FRHICommandListBase& RHICmdList, EResourceLockMode Loc
 			PendingLock.LockMode = LockMode; 
 			PendingLock.FirstLock = bIsFirstLock;
 
-			VulkanRHI::FStagingBuffer* StagingBuffer = Device->GetStagingManager().AcquireBuffer(LockSize, VK_BUFFER_USAGE_TRANSFER_SRC_BIT);
+			VulkanRHI::FStagingBuffer* StagingBuffer = Device.GetStagingManager().AcquireBuffer(LockSize, VK_BUFFER_USAGE_TRANSFER_SRC_BIT);
 			PendingLock.StagingBuffer = StagingBuffer;
 			Data = StagingBuffer->GetMappedPointer();
 
@@ -303,7 +301,7 @@ void* FVulkanBuffer::Lock(FRHICommandListBase& RHICmdList, EResourceLockMode Loc
 			RHICmdList.EnqueueLambda(TEXT("FVulkanBuffer::Lock"), [Buffer = this, NewAlloc](FRHICommandListBase& CmdList)
 			{
 				UpdateVulkanBufferStats(Buffer->GetDesc(), Buffer->CurrentBufferAlloc.Alloc.Size, false);
-				Buffer->CurrentBufferAlloc.Alloc.Free(*Buffer->GetParent());
+				Buffer->CurrentBufferAlloc.Alloc.Free(Buffer->Device);
 				Buffer->CurrentBufferAlloc = NewAlloc;
 				Buffer->CurrentBufferAlloc.Alloc.Own();
 				Buffer->UpdateLinkedViews();
@@ -327,7 +325,7 @@ void* FVulkanBuffer::Lock(FRHICommandListBase& RHICmdList, EResourceLockMode Loc
 
 void FVulkanBuffer::Unlock(FRHICommandListBase& RHICmdList)
 {
-	const bool bUnifiedMem = Device->HasUnifiedMemory();
+	const bool bUnifiedMem = Device.HasUnifiedMemory();
 	const bool bDynamic = EnumHasAnyFlags(GetUsage(), BUF_Dynamic) || EnumHasAnyFlags(GetUsage(), BUF_Volatile);
 	const bool bStatic = EnumHasAnyFlags(GetUsage(), BUF_Static) || !bDynamic;
 	const bool bSR = EnumHasAnyFlags(GetUsage(), BUF_ShaderResource);
@@ -353,7 +351,7 @@ void FVulkanBuffer::Unlock(FRHICommandListBase& RHICmdList)
 			if (PendingLock.LockMode == RLM_ReadOnly)
 			{
 				// Just remove the staging buffer here.
-				Buffer->Device->GetStagingManager().ReleaseBuffer(nullptr, StagingBuffer);
+				Buffer->Device.GetStagingManager().ReleaseBuffer(nullptr, StagingBuffer);
 			}
 			else if (PendingLock.LockMode == RLM_WriteOnly)
 			{
@@ -386,7 +384,7 @@ void FVulkanBuffer::Unlock(FRHICommandListBase& RHICmdList)
 				VkMemoryBarrier BarrierAfter = { VK_STRUCTURE_TYPE_MEMORY_BARRIER, nullptr, VK_ACCESS_TRANSFER_WRITE_BIT, VK_ACCESS_MEMORY_READ_BIT | VK_ACCESS_MEMORY_WRITE_BIT };
 				VulkanRHI::vkCmdPipelineBarrier(CommandBufferHandle, VK_PIPELINE_STAGE_TRANSFER_BIT, VK_PIPELINE_STAGE_ALL_COMMANDS_BIT, 0, 1, &BarrierAfter, 0, nullptr, 0, nullptr);
 
-				Buffer->GetParent()->GetStagingManager().ReleaseBuffer(CommonContext, StagingBuffer);
+				Buffer->Device.GetStagingManager().ReleaseBuffer(CommonContext, StagingBuffer);
 			}
 		});
 	}
@@ -397,7 +395,6 @@ void FVulkanBuffer::Unlock(FRHICommandListBase& RHICmdList)
 void FVulkanBuffer::TakeOwnership(FVulkanBuffer& Other)
 {
 	check(Other.LockStatus == ELockStatus::Unlocked);
-	check(GetParent() == Other.GetParent());
 
 	// Clean up any resource this buffer already owns
 	ReleaseOwnership();
@@ -419,7 +416,7 @@ void FVulkanBuffer::ReleaseOwnership()
 	if (CurrentBufferAlloc.Alloc.HasAllocation())
 	{
 		UpdateVulkanBufferStats(GetDesc(), CurrentBufferAlloc.Alloc.Size, false);
-		Device->GetMemoryManager().FreeVulkanAllocation(CurrentBufferAlloc.Alloc);
+		Device.GetMemoryManager().FreeVulkanAllocation(CurrentBufferAlloc.Alloc);
 	}
 
 	FRHIBuffer::ReleaseOwnership();
@@ -436,7 +433,7 @@ FRHIBufferInitializer FVulkanDynamicRHI::RHICreateBufferInitializer(FRHICommandL
 	LLM_SCOPE_DYNAMIC_STAT_OBJECTPATH_FNAME(CreateDesc.GetTraceClassName(), ELLMTagSet::AssetClasses);
 	UE_TRACE_METADATA_SCOPE_ASSET_FNAME(CreateDesc.DebugName, CreateDesc.GetTraceClassName(), CreateDesc.OwnerName);
 
-	FVulkanBuffer* Buffer = new FVulkanBuffer(Device, CreateDesc);
+	FVulkanBuffer* Buffer = new FVulkanBuffer(*Device, CreateDesc);
 
 	if (CreateDesc.IsNull() || CreateDesc.InitAction == ERHIBufferInitAction::Default)
 	{
diff --git a/Engine/Source/Runtime/VulkanRHI/Private/VulkanDevice.cpp b/Engine/Source/Runtime/VulkanRHI/Private/VulkanDevice.cpp
--- a/Engine/Source/Runtime/VulkanRHI/Private/VulkanDevice.cpp
+++ b/Engine/Source/Runtime/VulkanRHI/Private/VulkanDevice.cpp
@@ -230,7 +230,7 @@ void FVulkanPhysicalDeviceFeatures::Query(VkPhysicalDevice PhysicalDevice, uint3
 FVulkanDevice::FVulkanDevice(FVulkanDynamicRHI* InRHI, VkPhysicalDevice InGpu)
 	: Device(VK_NULL_HANDLE)
 	, MemoryManager(this)
-	, DeferredDeletionQueue(this)
+	, DeferredDeletionQueue(*this)
 	, Gpu(InGpu)
 	, ImmediateContext(nullptr)
 	, PipelineStateCache(nullptr)
@@ -1574,7 +1574,7 @@ FVulkanTransientHeapCache& FVulkanDevice::GetOrCreateTransientHeapCache()
 {
 	if (!TransientHeapCache)
 	{
-		TransientHeapCache = FVulkanTransientHeapCache::Create(this);
+		TransientHeapCache = FVulkanTransientHeapCache::Create(*this);
 	}
 	return *TransientHeapCache;
 }
diff --git a/Engine/Source/Runtime/VulkanRHI/Private/VulkanDevice.h b/Engine/Source/Runtime/VulkanRHI/Private/VulkanDevice.h
--- a/Engine/Source/Runtime/VulkanRHI/Private/VulkanDevice.h
+++ b/Engine/Source/Runtime/VulkanRHI/Private/VulkanDevice.h
@@ -196,11 +196,11 @@ private:
 
 namespace VulkanRHI
 {
-	class FDeferredDeletionQueue2 : public VulkanRHI::FDeviceChild
+	class FDeferredDeletionQueue2
 	{
 
 	public:
-		FDeferredDeletionQueue2(FVulkanDevice* InDevice);
+		FDeferredDeletionQueue2(FVulkanDevice& InDevice);
 		~FDeferredDeletionQueue2();
 
 		enum class EType
@@ -265,6 +265,7 @@ namespace VulkanRHI
 
 		void ReleaseResourcesImmediately(const TArray<FEntry>& InEntries);
 
+		FVulkanDevice& Device;
 		FCriticalSection CS;
 		TArray<FEntry> Entries;
 	};
diff --git a/Engine/Source/Runtime/VulkanRHI/Private/VulkanMemory.cpp b/Engine/Source/Runtime/VulkanRHI/Private/VulkanMemory.cpp
--- a/Engine/Source/Runtime/VulkanRHI/Private/VulkanMemory.cpp
+++ b/Engine/Source/Runtime/VulkanRHI/Private/VulkanMemory.cpp
@@ -2646,7 +2646,7 @@ namespace VulkanRHI
 
 	uint32 FMemoryManager::CalculateBufferAlignment(FVulkanDevice& InDevice, EBufferUsageFlags InUEUsage, bool bZeroSize)
 	{
-		const VkBufferUsageFlags VulkanBufferUsage = FVulkanBuffer::UEToVKBufferUsageFlags(&InDevice, InUEUsage, bZeroSize);
+		const VkBufferUsageFlags VulkanBufferUsage = FVulkanBuffer::UEToVKBufferUsageFlags(InDevice, InUEUsage, bZeroSize);
 
 		uint32 Alignment = CalculateBufferAlignmentFromVKUsageFlags(InDevice, VulkanBufferUsage);
 
@@ -4716,8 +4716,8 @@ namespace VulkanRHI
 
 
 
-	FDeferredDeletionQueue2::FDeferredDeletionQueue2(FVulkanDevice* InDevice)
-		: FDeviceChild(InDevice)
+	FDeferredDeletionQueue2::FDeferredDeletionQueue2(FVulkanDevice& InDevice)
+		: Device(InDevice)
 	{
 	}
 
@@ -4728,8 +4728,6 @@ namespace VulkanRHI
 
 	void FDeferredDeletionQueue2::EnqueueGenericResource(EType Type, uint64 Handle)
 	{
-		FVulkanQueue* Queue = Device->GetGraphicsQueue();
-
 		FEntry Entry;
 		Entry.StructureType = Type;
 		Entry.FrameNumber = GVulkanRHIDeletionFrameNumber;
@@ -4753,7 +4751,6 @@ namespace VulkanRHI
 	void FDeferredDeletionQueue2::EnqueueDeviceAllocation(FDeviceMemoryAllocation* DeviceMemoryAllocation)
 	{
 		check(DeviceMemoryAllocation);
-		FVulkanQueue* Queue = Device->GetGraphicsQueue();
 
 		FEntry Entry;
 		Entry.StructureType = EType::DeviceMemoryAllocation;
@@ -4774,7 +4771,6 @@ namespace VulkanRHI
 			return;
 		}
 		Allocation.Disown();
-		FVulkanQueue* Queue = Device->GetGraphicsQueue();
 
 		FEntry Entry;
 		Entry.StructureType = EType::ResourceAllocation;
@@ -4831,9 +4827,9 @@ namespace VulkanRHI
 
 			if (TempEntries.Num())
 			{
-				FVulkanDynamicRHI::Get().EnqueueEndOfPipeTask([Device = Device, Array = MoveTemp(TempEntries)]()
+				FVulkanDynamicRHI::Get().EnqueueEndOfPipeTask([&Device = Device, Array = MoveTemp(TempEntries)]()
 				{
-					Device->GetDeferredDeletionQueue().ReleaseResourcesImmediately(Array);
+					Device.GetDeferredDeletionQueue().ReleaseResourcesImmediately(Array);
 				});
 			}
 		}
@@ -4841,7 +4837,7 @@ namespace VulkanRHI
 
 	void FDeferredDeletionQueue2::ReleaseResourcesImmediately(const TArray<FEntry>& InEntries)
 	{
-		VkDevice DeviceHandle = Device->GetHandle();
+		VkDevice DeviceHandle = Device.GetHandle();
 
 		// Traverse list backwards so the swap switches to elements already tested
 		for (const FEntry& Entry : InEntries)
@@ -4867,14 +4863,14 @@ namespace VulkanRHI
 			{
 				FVulkanAllocation Allocation = Entry.Allocation;
 				Allocation.Own();
-				Device->GetMemoryManager().FreeVulkanAllocation(Allocation, EVulkanFreeFlag_DontDefer);
+				Device.GetMemoryManager().FreeVulkanAllocation(Allocation, EVulkanFreeFlag_DontDefer);
 				break;
 			}
 			case EType::DeviceMemoryAllocation:
 			{
 				check(Entry.DeviceMemoryAllocation);
 				FDeviceMemoryAllocation* DeviceMemoryAllocation = Entry.DeviceMemoryAllocation;
-				Device->GetDeviceMemoryManager().Free(DeviceMemoryAllocation);
+				Device.GetDeviceMemoryManager().Free(DeviceMemoryAllocation);
 				break;
 			}
 			case EType::AccelerationStructure:
@@ -4884,9 +4880,9 @@ namespace VulkanRHI
 			}
 			case EType::BindlessHandle:
 			{
-				check(Device->SupportsBindless());
+				check(Device.SupportsBindless());
 				const FRHIDescriptorHandle& DescriptorHandle = reinterpret_cast<const FRHIDescriptorHandle&>(Entry.Handle);
-				Device->GetBindlessDescriptorManager()->FreeDescriptor(DescriptorHandle);
+				Device.GetBindlessDescriptorManager()->FreeDescriptor(DescriptorHandle);
 				break;
 			}
 
diff --git a/Engine/Source/Runtime/VulkanRHI/Private/VulkanRayTracing.cpp b/Engine/Source/Runtime/VulkanRHI/Private/VulkanRayTracing.cpp
--- a/Engine/Source/Runtime/VulkanRHI/Private/VulkanRayTracing.cpp
+++ b/Engine/Source/Runtime/VulkanRHI/Private/VulkanRayTracing.cpp
@@ -510,7 +510,7 @@ void FVulkanRayTracingGeometry::CompactAccelerationStructure(FVulkanCommandBuffe
 		FRHIBufferCreateDesc::Create(*DebugNameString, InSizeAfterCompaction, 0, BUF_AccelerationStructure)
 		.SetInitialState(ERHIAccess::BVHWrite);
 
-	AccelerationStructureBuffer = new FVulkanBuffer(Device, CreateDesc);
+	AccelerationStructureBuffer = new FVulkanBuffer(*Device, CreateDesc);
 
 	INC_MEMORY_STAT_BY(STAT_VulkanRayTracingUsedVideoMemory, AccelerationStructureBuffer->GetSize());
 	INC_MEMORY_STAT_BY(STAT_VulkanRayTracingBLASMemory, AccelerationStructureBuffer->GetSize());
diff --git a/Engine/Source/Runtime/VulkanRHI/Private/VulkanResources.h b/Engine/Source/Runtime/VulkanRHI/Private/VulkanResources.h
--- a/Engine/Source/Runtime/VulkanRHI/Private/VulkanResources.h
+++ b/Engine/Source/Runtime/VulkanRHI/Private/VulkanResources.h
@@ -809,10 +809,10 @@ protected:
 };
 
 
-class FVulkanBuffer : public FRHIBuffer, public VulkanRHI::FDeviceChild, public FVulkanViewableResource
+class FVulkanBuffer : public FRHIBuffer, public FVulkanViewableResource
 {
 public:
-	FVulkanBuffer(FVulkanDevice* InDevice, const FRHIBufferCreateDesc& CreateDesc, const FRHITransientHeapAllocation* InTransientHeapAllocation = nullptr);
+	FVulkanBuffer(FVulkanDevice& InDevice, const FRHIBufferCreateDesc& CreateDesc, const FRHITransientHeapAllocation* InTransientHeapAllocation = nullptr);
 	virtual ~FVulkanBuffer();
 
 	inline const VulkanRHI::FVulkanAllocation& GetCurrentAllocation() const
@@ -869,7 +869,7 @@ public:
 		Callback(TEXT("FVulkanBuffer"), FName(), this, 0, GetCurrentSize(), 1, 1, VK_FORMAT_UNDEFINED);
 	}
 
-	static VkBufferUsageFlags UEToVKBufferUsageFlags(FVulkanDevice* InDevice, EBufferUsageFlags InUEUsage, bool bZeroSize);
+	static VkBufferUsageFlags UEToVKBufferUsageFlags(FVulkanDevice& InDevice, EBufferUsageFlags InUEUsage, bool bZeroSize);
 
 	struct FBufferAlloc
 	{
@@ -907,6 +907,8 @@ protected:
 	// Will return a new allocation that can be used for this buffer (proper memory type and size)
 	void AllocateMemory(FBufferAlloc& OutAlloc);
 
+	FVulkanDevice& Device;
+
 	VkBufferUsageFlags BufferUsageFlags;
 
 	enum class ELockStatus : uint8
diff --git a/Engine/Source/Runtime/VulkanRHI/Private/VulkanTransientResourceAllocator.cpp b/Engine/Source/Runtime/VulkanRHI/Private/VulkanTransientResourceAllocator.cpp
--- a/Engine/Source/Runtime/VulkanRHI/Private/VulkanTransientResourceAllocator.cpp
+++ b/Engine/Source/Runtime/VulkanRHI/Private/VulkanTransientResourceAllocator.cpp
@@ -7,14 +7,14 @@
 #include "VulkanDevice.h"
 #include "VulkanMemory.h"
 
-FVulkanTransientHeap::FVulkanTransientHeap(const FInitializer& Initializer, FVulkanDevice* InDevice)
+FVulkanTransientHeap::FVulkanTransientHeap(const FInitializer& Initializer, FVulkanDevice& InDevice)
 	: FRHITransientHeap(Initializer)
-	, FDeviceChild(InDevice)
+	, Device(InDevice)
 {
 	EBufferUsageFlags UEBufferUsageFlags = BUF_UniformBuffer | BUF_VertexBuffer | BUF_IndexBuffer | BUF_DrawIndirect 
 		| BUF_UnorderedAccess | BUF_StructuredBuffer | BUF_ShaderResource | BUF_KeepCPUAccessible;
 
-	if (InDevice->GetOptionalExtensions().HasRaytracingExtensions())
+	if (InDevice.GetOptionalExtensions().HasRaytracingExtensions())
 	{
 		UEBufferUsageFlags |= BUF_RayTracingScratch;
 		// AccelerationStructure not yet supported as TransientResource see FVulkanTransientResourceAllocator::CreateBuffer
@@ -23,30 +23,30 @@ FVulkanTransientHeap::FVulkanTransientHeap(const FInitializer& Initializer, FVul
 
 	const bool bZeroSize = false;
 	const VkBufferUsageFlags BufferUsageFlags = FVulkanBuffer::UEToVKBufferUsageFlags(InDevice, UEBufferUsageFlags, bZeroSize);
-	VulkanBuffer = InDevice->CreateBuffer(Initializer.Size, BufferUsageFlags);
+	VulkanBuffer = InDevice.CreateBuffer(Initializer.Size, BufferUsageFlags);
 	if (UseVulkanDescriptorCache())
 	{
 		HandleID = ++GVulkanBufferHandleIdCounter;
 	}
 
 	// Find the alignment that works for everyone
-	const uint32 MinBufferAlignment = FMath::Max<uint32>(Initializer.Alignment, VulkanRHI::FMemoryManager::CalculateBufferAlignment(*InDevice, UEBufferUsageFlags, bZeroSize));
+	const uint32 MinBufferAlignment = FMath::Max<uint32>(Initializer.Alignment, VulkanRHI::FMemoryManager::CalculateBufferAlignment(InDevice, UEBufferUsageFlags, bZeroSize));
 
 	const VulkanRHI::EVulkanAllocationFlags AllocFlags = VulkanRHI::EVulkanAllocationFlags::Dedicated | VulkanRHI::EVulkanAllocationFlags::AutoBind;
-	InDevice->GetMemoryManager().AllocateBufferMemory(InternalAllocation, VulkanBuffer, AllocFlags, TEXT("VulkanTransientHeap"), MinBufferAlignment);
+	InDevice.GetMemoryManager().AllocateBufferMemory(InternalAllocation, VulkanBuffer, AllocFlags, TEXT("VulkanTransientHeap"), MinBufferAlignment);
 }
 
 FVulkanTransientHeap::~FVulkanTransientHeap()
 {
-	Device->GetMemoryManager().FreeVulkanAllocation(InternalAllocation);
-	Device->GetDeferredDeletionQueue().EnqueueResource(VulkanRHI::FDeferredDeletionQueue2::EType::Buffer, VulkanBuffer);
+	Device.GetMemoryManager().FreeVulkanAllocation(InternalAllocation);
+	Device.GetDeferredDeletionQueue().EnqueueResource(VulkanRHI::FDeferredDeletionQueue2::EType::Buffer, VulkanBuffer);
 	VulkanBuffer = VK_NULL_HANDLE;
 	HandleID = 0;
 }
 
 VkDeviceMemory FVulkanTransientHeap::GetMemoryHandle()
 {
-	return InternalAllocation.GetDeviceMemoryHandle(Device);
+	return InternalAllocation.GetDeviceMemoryHandle(&Device);
 }
 
 VulkanRHI::FVulkanAllocation FVulkanTransientHeap::GetVulkanAllocation(const FRHITransientHeapAllocation& HeapAllocation)
@@ -65,12 +65,12 @@ VulkanRHI::FVulkanAllocation FVulkanTransientHeap::GetVulkanAllocation(const FRH
 	return TransientAlloc;
 }
 
-FVulkanTransientHeapCache* FVulkanTransientHeapCache::Create(FVulkanDevice* InDevice)
+FVulkanTransientHeapCache* FVulkanTransientHeapCache::Create(FVulkanDevice& InDevice)
 {
 	FRHITransientHeapCache::FInitializer Initializer = FRHITransientHeapCache::FInitializer::CreateDefault();
 
 	// Respect a minimum alignment
-	Initializer.HeapAlignment = FMath::Max((uint32)InDevice->GetLimits().bufferImageGranularity, 256u);
+	Initializer.HeapAlignment = FMath::Max((uint32)InDevice.GetLimits().bufferImageGranularity, 256u);
 
 	// Mix resource types onto the same heap.
 	Initializer.bSupportsAllHeapFlags = true;
@@ -78,9 +78,9 @@ FVulkanTransientHeapCache* FVulkanTransientHeapCache::Create(FVulkanDevice* InDe
 	return new FVulkanTransientHeapCache(Initializer, InDevice);
 }
 
-FVulkanTransientHeapCache::FVulkanTransientHeapCache(const FRHITransientHeapCache::FInitializer& Initializer, FVulkanDevice* InDevice)
+FVulkanTransientHeapCache::FVulkanTransientHeapCache(const FRHITransientHeapCache::FInitializer& Initializer, FVulkanDevice& InDevice)
 	: FRHITransientHeapCache(Initializer)
-	, FDeviceChild(InDevice)
+	, Device(InDevice)
 {
 }
 
@@ -92,7 +92,7 @@ FRHITransientHeap* FVulkanTransientHeapCache::CreateHeap(const FRHITransientHeap
 
 FVulkanTransientResourceAllocator::FVulkanTransientResourceAllocator(FVulkanTransientHeapCache& InHeapCache)
 	: FRHITransientResourceHeapAllocator(InHeapCache)
-	, FDeviceChild(InHeapCache.GetParent())
+	, Device(InHeapCache.Device)
 {
 }
 
@@ -101,7 +101,7 @@ FRHITransientTexture* FVulkanTransientResourceAllocator::CreateTexture(const FRH
 	FDynamicRHI::FRHICalcTextureSizeResult MemReq = GDynamicRHI->RHICalcTexturePlatformSize(InCreateInfo, 0);
 
 	return CreateTextureInternal(InCreateInfo, InDebugName, Fences, MemReq.Size, MemReq.Align,
-		[Device=Device, InCreateInfo, InDebugName, MemReq](const FRHITransientHeap::FResourceInitializer& Initializer)
+		[InCreateInfo, InDebugName, MemReq](const FRHITransientHeap::FResourceInitializer& Initializer)
 	{
 		FRHITextureCreateDesc CreateDesc(InCreateInfo, ERHIAccess::Discard, InDebugName);
 
@@ -120,11 +120,11 @@ FRHITransientBuffer* FVulkanTransientResourceAllocator::CreateBuffer(const FRHIB
 	checkf(!EnumHasAnyFlags(CreateDesc.Usage, BUF_Volatile), TEXT("The volatile flag is not supported for transient resources."));
 
 	const bool bZeroSize = (CreateDesc.Size == 0);
-	const uint32 Alignment = VulkanRHI::FMemoryManager::CalculateBufferAlignment(*Device, CreateDesc.Usage, bZeroSize);
+	const uint32 Alignment = VulkanRHI::FMemoryManager::CalculateBufferAlignment(Device, CreateDesc.Usage, bZeroSize);
 	uint64 Size = Align(CreateDesc.Size, Alignment);
 
 	return CreateBufferInternal(CreateDesc, InDebugName, Fences, Size, Alignment,
-	[Device=Device, CreateDesc, InDebugName, Size](const FRHITransientHeap::FResourceInitializer& Initializer)
+	[&Device=Device, CreateDesc, InDebugName, Size](const FRHITransientHeap::FResourceInitializer& Initializer)
 	{
 		FRHIBuffer* Buffer = new FVulkanBuffer(Device, CreateDesc, &Initializer.Allocation);
 		return new FRHITransientBuffer(Buffer, 0/*GpuVirtualAddress*/, Initializer.Hash, Size, ERHITransientAllocationType::Heap, CreateDesc);
diff --git a/Engine/Source/Runtime/VulkanRHI/Private/VulkanTransientResourceAllocator.h b/Engine/Source/Runtime/VulkanRHI/Private/VulkanTransientResourceAllocator.h
--- a/Engine/Source/Runtime/VulkanRHI/Private/VulkanTransientResourceAllocator.h
+++ b/Engine/Source/Runtime/VulkanRHI/Private/VulkanTransientResourceAllocator.h
@@ -9,16 +9,16 @@
 class FVulkanTransientHeap final
 	: public FRHITransientHeap
 	, public FRefCountBase
-	, public VulkanRHI::FDeviceChild
 {
 public:
-	FVulkanTransientHeap(const FInitializer& Initializer, FVulkanDevice* InDevice);
+	FVulkanTransientHeap(const FInitializer& Initializer, FVulkanDevice& InDevice);
 	~FVulkanTransientHeap();
 
 	VkDeviceMemory GetMemoryHandle();
 	static VulkanRHI::FVulkanAllocation GetVulkanAllocation(const FRHITransientHeapAllocation& TransientInitializer);
 
 private:
+	FVulkanDevice& Device;
 	VkBuffer VulkanBuffer = VK_NULL_HANDLE;
 	uint32 HandleID = 0;
 	VkMemoryRequirements MemoryRequirements;
@@ -28,21 +28,21 @@ private:
 
 class FVulkanTransientHeapCache final
 	: public FRHITransientHeapCache
-	, public VulkanRHI::FDeviceChild
 {
 public:
-	static FVulkanTransientHeapCache* Create(FVulkanDevice* InDevice);
+	static FVulkanTransientHeapCache* Create(FVulkanDevice& InDevice);
 
 	FRHITransientHeap* CreateHeap(const FRHITransientHeap::FInitializer& Initializer) override;
 
+	FVulkanDevice& Device;
+
 private:
-	FVulkanTransientHeapCache(const FInitializer& Initializer, FVulkanDevice* InDevice);
+	FVulkanTransientHeapCache(const FInitializer& Initializer, FVulkanDevice& InDevice);
 };
 
 
 class FVulkanTransientResourceAllocator final
 	: public FRHITransientResourceHeapAllocator
-	, public VulkanRHI::FDeviceChild
 {
 public:
 	FVulkanTransientResourceAllocator(FVulkanTransientHeapCache& InHeapCache);
@@ -59,4 +59,6 @@ public:
 	}
 	FRHITransientTexture* CreateTexture(const FRHITextureCreateInfo& InCreateInfo, const TCHAR* InDebugName, const FRHITransientAllocationFences& Fences) override;
 	FRHITransientBuffer* CreateBuffer(const FRHIBufferCreateInfo& InCreateInfo, const TCHAR* InDebugName, const FRHITransientAllocationFences& Fences) override;
+
+	FVulkanDevice& Device;
 };
\ No newline at end of file
