VulkanRHI: Remove remaining uses of VulkanRHI::FDeviceChild

#rnx

diff --git a/Engine/Source/Runtime/VulkanRHI/Private/VulkanBindlessDescriptorManager.cpp b/Engine/Source/Runtime/VulkanRHI/Private/VulkanBindlessDescriptorManager.cpp
--- a/Engine/Source/Runtime/VulkanRHI/Private/VulkanBindlessDescriptorManager.cpp
+++ b/Engine/Source/Runtime/VulkanRHI/Private/VulkanBindlessDescriptorManager.cpp
@@ -378,7 +378,7 @@ void FVulkanBindlessDescriptorManager::Init()
 	const VkPhysicalDeviceDescriptorBufferPropertiesEXT& DescriptorBufferProperties = Device.GetOptionalExtensionProperties().DescriptorBufferProps;
 
 	const VkBufferUsageFlags BufferUsageFlags = VK_BUFFER_USAGE_SHADER_DEVICE_ADDRESS_BIT | VK_BUFFER_USAGE_RESOURCE_DESCRIPTOR_BUFFER_BIT_EXT;
-	SingleUseUBAllocator = new VulkanRHI::FTempBlockAllocator(&Device, GVulkanBindlessBlockSize, DescriptorBufferProperties.descriptorBufferOffsetAlignment, BufferUsageFlags);
+	SingleUseUBAllocator = new VulkanRHI::FTempBlockAllocator(Device, GVulkanBindlessBlockSize, DescriptorBufferProperties.descriptorBufferOffsetAlignment, BufferUsageFlags);
 
 	// Create the dummy layout for unsupported descriptor types
 	{
diff --git a/Engine/Source/Runtime/VulkanRHI/Private/VulkanDevice.cpp b/Engine/Source/Runtime/VulkanRHI/Private/VulkanDevice.cpp
--- a/Engine/Source/Runtime/VulkanRHI/Private/VulkanDevice.cpp
+++ b/Engine/Source/Runtime/VulkanRHI/Private/VulkanDevice.cpp
@@ -229,7 +229,7 @@ void FVulkanPhysicalDeviceFeatures::Query(VkPhysicalDevice PhysicalDevice, uint3
 
 FVulkanDevice::FVulkanDevice(FVulkanDynamicRHI* InRHI, VkPhysicalDevice InGpu)
 	: Device(VK_NULL_HANDLE)
-	, MemoryManager(this)
+	, MemoryManager(*this)
 	, DeferredDeletionQueue(*this)
 	, Gpu(InGpu)
 	, ImmediateContext(nullptr)
@@ -1364,7 +1364,7 @@ void FVulkanDevice::InitGPU()
 		}
 
 		const int32 VulkanTempBlockSize = GCVarVulkanTempBlockSizeKB.GetValueOnAnyThread() * 1024;
-		TempBlockAllocator = new VulkanRHI::FTempBlockAllocator(this, VulkanTempBlockSize, BlockAlignment, BufferUsageFlags);
+		TempBlockAllocator = new VulkanRHI::FTempBlockAllocator(*this, VulkanTempBlockSize, BlockAlignment, BufferUsageFlags);
 	}
 
 	RenderPassManager = new FVulkanRenderPassManager(*this);
@@ -1413,7 +1413,7 @@ void FVulkanDevice::InitGPU()
 	if (RHISupportsRayTracing(GMaxRHIShaderPlatform) && GetOptionalExtensions().HasRaytracingExtensions())
 	{
 		check(RayTracingCompactionRequestHandler == nullptr);
-		RayTracingCompactionRequestHandler = new FVulkanRayTracingCompactionRequestHandler(this);
+		RayTracingCompactionRequestHandler = new FVulkanRayTracingCompactionRequestHandler(*this);
 	}
 
 	FVulkanPlatform::PostInitGPU(*this);
diff --git a/Engine/Source/Runtime/VulkanRHI/Private/VulkanLayout.cpp b/Engine/Source/Runtime/VulkanRHI/Private/VulkanLayout.cpp
--- a/Engine/Source/Runtime/VulkanRHI/Private/VulkanLayout.cpp
+++ b/Engine/Source/Runtime/VulkanRHI/Private/VulkanLayout.cpp
@@ -3,9 +3,9 @@
 #include "VulkanLayout.h"
 #include "VulkanBindlessDescriptorManager.h"
 
-FVulkanLayout::FVulkanLayout(FVulkanDevice* InDevice, bool InGfxLayout, bool InUsesBindless)
-	: VulkanRHI::FDeviceChild(InDevice)
-	, DescriptorSetLayout(Device)
+FVulkanLayout::FVulkanLayout(FVulkanDevice& InDevice, bool InGfxLayout, bool InUsesBindless)
+	: Device(InDevice)
+	, DescriptorSetLayout(&InDevice)
 	, PipelineLayout(VK_NULL_HANDLE)
 	, bIsGfxLayout(InGfxLayout)
 	, bUsesBindless(InUsesBindless)
@@ -16,7 +16,7 @@ FVulkanLayout::~FVulkanLayout()
 {
 	if (!bUsesBindless && PipelineLayout != VK_NULL_HANDLE)
 	{
-		Device->GetDeferredDeletionQueue().EnqueueResource(VulkanRHI::FDeferredDeletionQueue2::EType::PipelineLayout, PipelineLayout);
+		Device.GetDeferredDeletionQueue().EnqueueResource(VulkanRHI::FDeferredDeletionQueue2::EType::PipelineLayout, PipelineLayout);
 		PipelineLayout = VK_NULL_HANDLE;
 	}
 }
@@ -30,7 +30,7 @@ void FVulkanLayout::Compile(FVulkanDescriptorSetLayoutMap& DSetLayoutMap)
 	if (bUsesBindless)
 	{
 		// Get the bindless manager's layout
-		PipelineLayout = Device->GetBindlessDescriptorManager()->GetPipelineLayout();
+		PipelineLayout = Device.GetBindlessDescriptorManager()->GetPipelineLayout();
 	}
 	else
 	{
@@ -41,6 +41,6 @@ void FVulkanLayout::Compile(FVulkanDescriptorSetLayoutMap& DSetLayoutMap)
 		PipelineLayoutCreateInfo.setLayoutCount = LayoutHandles.Num();
 		PipelineLayoutCreateInfo.pSetLayouts = LayoutHandles.GetData();
 		//PipelineLayoutCreateInfo.pushConstantRangeCount = 0;
-		VERIFYVULKANRESULT(VulkanRHI::vkCreatePipelineLayout(Device->GetHandle(), &PipelineLayoutCreateInfo, VULKAN_CPU_ALLOCATOR, &PipelineLayout));
+		VERIFYVULKANRESULT(VulkanRHI::vkCreatePipelineLayout(Device.GetHandle(), &PipelineLayoutCreateInfo, VULKAN_CPU_ALLOCATOR, &PipelineLayout));
 	}
 }
diff --git a/Engine/Source/Runtime/VulkanRHI/Private/VulkanLayout.h b/Engine/Source/Runtime/VulkanRHI/Private/VulkanLayout.h
--- a/Engine/Source/Runtime/VulkanRHI/Private/VulkanLayout.h
+++ b/Engine/Source/Runtime/VulkanRHI/Private/VulkanLayout.h
@@ -6,10 +6,10 @@
 #include "VulkanDescriptorSets.h"
 
 // Layout for a Pipeline, also includes DescriptorSets layout
-class FVulkanLayout : public VulkanRHI::FDeviceChild
+class FVulkanLayout
 {
 public:
-	FVulkanLayout(FVulkanDevice* InDevice, bool InGfxLayout, bool InUsesBindless);
+	FVulkanLayout(FVulkanDevice& InDevice, bool InGfxLayout, bool InUsesBindless);
 	virtual ~FVulkanLayout();
 
 	const FVulkanDescriptorSetsLayout& GetDescriptorSetsLayout() const
@@ -38,6 +38,7 @@ public:
 	}
 
 protected:
+	FVulkanDevice& Device;
 	FVulkanDescriptorSetsLayout	DescriptorSetLayout;
 	VkPipelineLayout			PipelineLayout;
 
diff --git a/Engine/Source/Runtime/VulkanRHI/Private/VulkanMemory.cpp b/Engine/Source/Runtime/VulkanRHI/Private/VulkanMemory.cpp
--- a/Engine/Source/Runtime/VulkanRHI/Private/VulkanMemory.cpp
+++ b/Engine/Source/Runtime/VulkanRHI/Private/VulkanMemory.cpp
@@ -1599,7 +1599,7 @@ namespace VulkanRHI
 	FVulkanResourceHeap::FVulkanResourceHeap(FMemoryManager* InOwner, uint32 InMemoryTypeIndex, uint32 InOverridePageSize)
 		: Owner(InOwner)
 		, MemoryTypeIndex((uint16)InMemoryTypeIndex)
-		, HeapIndex((uint16)InOwner->GetParent()->GetDeviceMemoryManager().GetHeapIndex(InMemoryTypeIndex))
+		, HeapIndex((uint16)InOwner->GetDevice().GetDeviceMemoryManager().GetHeapIndex(InMemoryTypeIndex))
 		, bIsHostCachedSupported(false)
 		, bIsLazilyAllocatedSupported(false)
 		, OverridePageSize(InOverridePageSize)
@@ -1618,7 +1618,7 @@ namespace VulkanRHI
 			{
 				FVulkanSubresourceAllocator* Page = UsedPages[Index];
 				bLeak |= !Page->JoinFreeBlocks();
-				Owner->GetParent()->GetDeviceMemoryManager().Free(Page->MemoryAllocation);
+				Owner->GetDevice().GetDeviceMemoryManager().Free(Page->MemoryAllocation);
 				delete Page;
 			}
 			UsedPages.Reset(0);
@@ -1631,7 +1631,7 @@ namespace VulkanRHI
 		}
 		if (bDump)
 		{
-			Owner->GetParent()->GetMemoryManager().DumpMemory();
+			Owner->GetDevice().GetMemoryManager().DumpMemory();
 			GLog->Flush();
 		}
 	}
@@ -1640,7 +1640,7 @@ namespace VulkanRHI
 		Owner->UnregisterSubresourceAllocator(InPage);
 		FDeviceMemoryAllocation* Allocation = InPage->MemoryAllocation;
 		InPage->MemoryAllocation = 0;
-		Owner->GetParent()->GetDeviceMemoryManager().Free(Allocation );
+		Owner->GetDevice().GetDeviceMemoryManager().Free(Allocation );
 		UsedMemory -= InPage->MaxSize;
 		delete InPage;
 	}
@@ -1906,7 +1906,7 @@ namespace VulkanRHI
 	//Try to reallocate an allocation on a different page. Used When defragging
 	bool FVulkanResourceHeap::TryRealloc(FVulkanAllocation& OutAllocation, FVulkanEvictable* AllocationOwner, EType Type, uint32 Size, uint32 Alignment, EVulkanAllocationMetaType MetaType)
 	{
-		FDeviceMemoryManager& DeviceMemoryManager = Owner->GetParent()->GetDeviceMemoryManager();
+		FDeviceMemoryManager& DeviceMemoryManager = Owner->GetDevice().GetDeviceMemoryManager();
 		FVulkanPageSizeBucket MemoryBucket;
 		uint32 BucketId = GetPageSizeBucket(MemoryBucket, Type, Size, false);
 
@@ -1947,7 +1947,7 @@ namespace VulkanRHI
 
 		bForceSeparateAllocation = bForceSeparateAllocation || GVulkanSingleAllocationPerResource != 0;
 
-		FDeviceMemoryManager& DeviceMemoryManager = Owner->GetParent()->GetDeviceMemoryManager();
+		FDeviceMemoryManager& DeviceMemoryManager = Owner->GetDevice().GetDeviceMemoryManager();
 		FVulkanPageSizeBucket MemoryBucket;
 		uint32 BucketId = GetPageSizeBucket(MemoryBucket, Type, Size, bForceSeparateAllocation);
 
@@ -2045,7 +2045,7 @@ namespace VulkanRHI
 		VkMemoryDedicatedAllocateInfoKHR DedicatedAllocInfo;
 		ZeroVulkanStruct(DedicatedAllocInfo, VK_STRUCTURE_TYPE_MEMORY_DEDICATED_ALLOCATE_INFO_KHR);
 		DedicatedAllocInfo.image = Image;
-		FDeviceMemoryAllocation* DeviceMemoryAllocation = Owner->GetParent()->GetDeviceMemoryManager().Alloc(true, AllocationSize, MemoryTypeIndex, &DedicatedAllocInfo, VULKAN_MEMORY_HIGHEST_PRIORITY, bExternal, File, Line);
+		FDeviceMemoryAllocation* DeviceMemoryAllocation = Owner->GetDevice().GetDeviceMemoryManager().Alloc(true, AllocationSize, MemoryTypeIndex, &DedicatedAllocInfo, VULKAN_MEMORY_HIGHEST_PRIORITY, bExternal, File, Line);
 		if(!DeviceMemoryAllocation)
 		{
 			return false;
@@ -2074,9 +2074,9 @@ namespace VulkanRHI
 #endif
 	}
 
-	FMemoryManager::FMemoryManager(FVulkanDevice* InDevice)
-		: FDeviceChild(InDevice)
-		, DeviceMemoryManager(&InDevice->GetDeviceMemoryManager())
+	FMemoryManager::FMemoryManager(FVulkanDevice& InDevice)
+		: Device(InDevice)
+		, DeviceMemoryManager(&InDevice.GetDeviceMemoryManager())
 		, AllBufferAllocationsFreeListHead(-1)
 	{
 	}
@@ -2215,8 +2215,8 @@ namespace VulkanRHI
 					UE_LOG(LogVulkanRHI, Warning, TEXT("Suballocation(s) for Buffer %p were not released.%s"), (void*)BufferAllocation->Buffer, *VULKAN_TRACK_STRING(BufferAllocation->Track));
 				}
 
-				BufferAllocation->Destroy(GetParent());
-				GetParent()->GetDeviceMemoryManager().Free(BufferAllocation->MemoryAllocation);
+				BufferAllocation->Destroy(&Device);
+				Device.GetDeviceMemoryManager().Free(BufferAllocation->MemoryAllocation);
 				delete BufferAllocation;
 			}
 			UsedAllocations.Empty(0);
@@ -2227,8 +2227,8 @@ namespace VulkanRHI
 			for (int32 Index = 0; Index < FreeAllocations.Num(); ++Index)
 			{
 				FVulkanSubresourceAllocator* BufferAllocation = FreeAllocations[Index];
-				BufferAllocation->Destroy(GetParent());
-				GetParent()->GetDeviceMemoryManager().Free(BufferAllocation->MemoryAllocation);
+				BufferAllocation->Destroy(&Device);
+				Device.GetDeviceMemoryManager().Free(BufferAllocation->MemoryAllocation);
 				delete BufferAllocation;
 			}
 			FreeAllocations.Empty(0);
@@ -2258,7 +2258,7 @@ namespace VulkanRHI
 		{
 			FVulkanSubresourceAllocator* BufferAllocation = BufferAllocationsToRelease[Index];
 
-			BufferAllocation->Destroy(GetParent());
+			BufferAllocation->Destroy(&Device);
 
 			FVulkanResourceHeap* Heap = ResourceTypeHeaps[BufferAllocation->MemoryTypeIndex];
 			Heap->ReleasePage(BufferAllocation);
@@ -2343,7 +2343,7 @@ namespace VulkanRHI
 		if(BestEvictHeap && ((GVulkanEvictOnePage || UpdateEvictThreshold(true)) && PrimaryHeapIndex >= 0))
 		{
 			GVulkanEvictOnePage = 0;
-			PendingEvictBytes += BestEvictHeap->EvictOne(*Device, Context);
+			PendingEvictBytes += BestEvictHeap->EvictOne(Device, Context);
 		}
 
 #if !PLATFORM_ANDROID
@@ -2359,7 +2359,7 @@ namespace VulkanRHI
 			{
 				DumpMemory();
 			}
-			BestDefragHeap->DefragTick(*Device, Context, Count);
+			BestDefragHeap->DefragTick(Device, Context, Count);
 
 			if(GVulkanDefragOnce)
 			{
@@ -2495,7 +2495,7 @@ namespace VulkanRHI
 		else
 		{
 			GetSubresourceAllocator(Allocation.AllocatorIndex)->SetFreePending(Allocation);
-			Device->GetDeferredDeletionQueue().EnqueueResourceAllocation(Allocation);
+			Device.GetDeferredDeletionQueue().EnqueueResourceAllocation(Allocation);
 		}
 		check(!Allocation.HasAllocation());
 	}
@@ -2735,7 +2735,7 @@ namespace VulkanRHI
 		SCOPED_NAMED_EVENT(FResourceHeapManager_AllocateBufferPooled, FColor::Cyan);
 		check(OutAllocation.Type == EVulkanAllocationEmpty);
 
-		uint32 Alignment = FMath::Max(MinAlignment, CalculateBufferAlignmentFromVKUsageFlags(*Device, BufferUsageFlags));
+		uint32 Alignment = FMath::Max(MinAlignment, CalculateBufferAlignmentFromVKUsageFlags(Device, BufferUsageFlags));
 		const float Priority = CalculateBufferPriority(BufferUsageFlags);
 
 		const int32 PoolSize = (int32)GetPoolTypeForAlloc(Size, Alignment);
@@ -2779,15 +2779,15 @@ namespace VulkanRHI
 
 		// New Buffer
 		const uint32 BufferSize = FMath::Max(Size, BufferSizes[PoolSize]);
-		VkBuffer Buffer = Device->CreateBuffer(BufferSize, BufferUsageFlags);
+		VkBuffer Buffer = Device.CreateBuffer(BufferSize, BufferUsageFlags);
 
 		VkMemoryRequirements MemReqs;
-		VulkanRHI::vkGetBufferMemoryRequirements(Device->GetHandle(), Buffer, &MemReqs);
+		VulkanRHI::vkGetBufferMemoryRequirements(Device.GetHandle(), Buffer, &MemReqs);
 		Alignment = FMath::Max((uint32)MemReqs.alignment, Alignment);
 		ensure(MemReqs.size >= BufferSize);
 
 		uint32 MemoryTypeIndex;	
-		VERIFYVULKANRESULT(Device->GetDeviceMemoryManager().GetMemoryTypeFromProperties(MemReqs.memoryTypeBits, MemoryPropertyFlags, &MemoryTypeIndex));		
+		VERIFYVULKANRESULT(Device.GetDeviceMemoryManager().GetMemoryTypeFromProperties(MemReqs.memoryTypeBits, MemoryPropertyFlags, &MemoryTypeIndex));
 
 		bool bHasUnifiedMemory = DeviceMemoryManager->HasUnifiedMemory();
 		FDeviceMemoryAllocation* DeviceMemoryAllocation = DeviceMemoryManager->Alloc(true, MemReqs.size, MemoryTypeIndex, nullptr, Priority, false, File, Line);
@@ -2795,7 +2795,7 @@ namespace VulkanRHI
 		{
 			MemoryPropertyFlags &= (~VK_MEMORY_PROPERTY_DEVICE_LOCAL_BIT);
 			const uint32 ForbiddenMemoryTypeIndex = MemoryTypeIndex;
-			if (GVulkanMemoryMemoryFallbackToHost && VK_SUCCESS == Device->GetDeviceMemoryManager().GetMemoryTypeFromPropertiesExcluding(MemReqs.memoryTypeBits, MemoryPropertyFlags, ForbiddenMemoryTypeIndex, &MemoryTypeIndex))
+			if (GVulkanMemoryMemoryFallbackToHost && VK_SUCCESS == Device.GetDeviceMemoryManager().GetMemoryTypeFromPropertiesExcluding(MemReqs.memoryTypeBits, MemoryPropertyFlags, ForbiddenMemoryTypeIndex, &MemoryTypeIndex))
 			{
 				DeviceMemoryAllocation = DeviceMemoryManager->Alloc(false, MemReqs.size, MemoryTypeIndex, nullptr, Priority, false, File, Line);
 			}
@@ -2805,7 +2805,7 @@ namespace VulkanRHI
 			HandleOOM(false);
 			checkNoEntry();
 		}
-		VERIFYVULKANRESULT(VulkanRHI::vkBindBufferMemory(Device->GetHandle(), Buffer, DeviceMemoryAllocation->GetHandle(), 0));
+		VERIFYVULKANRESULT(VulkanRHI::vkBindBufferMemory(Device.GetHandle(), Buffer, DeviceMemoryAllocation->GetHandle(), 0));
 		uint8 AllocationFlags = 0;
 		if(!bHasUnifiedMemory && MetaTypeCanEvict(MetaType))
 		{
@@ -3048,7 +3048,7 @@ namespace VulkanRHI
 		ZeroVulkanStruct(MemoryRequirements, VK_STRUCTURE_TYPE_MEMORY_REQUIREMENTS_2);
 		MemoryRequirements.pNext = &DedicatedRequirements;
 
-		VulkanRHI::vkGetBufferMemoryRequirements2(Device->GetHandle(), &BufferMemoryRequirementsInfo, &MemoryRequirements);
+		VulkanRHI::vkGetBufferMemoryRequirements2(Device.GetHandle(), &BufferMemoryRequirementsInfo, &MemoryRequirements);
 
 		// Allow caller to force his own alignment requirements
 		MemoryRequirements.memoryRequirements.alignment = FMath::Max(MemoryRequirements.memoryRequirements.alignment, (VkDeviceSize)InForceMinAlignment);
@@ -3071,15 +3071,15 @@ namespace VulkanRHI
 				VkBindBufferMemoryInfo BindBufferMemoryInfo;
 				ZeroVulkanStruct(BindBufferMemoryInfo, VK_STRUCTURE_TYPE_BIND_BUFFER_MEMORY_INFO);
 				BindBufferMemoryInfo.buffer = InBuffer;
-				BindBufferMemoryInfo.memory = OutAllocation.GetDeviceMemoryHandle(Device);
+				BindBufferMemoryInfo.memory = OutAllocation.GetDeviceMemoryHandle(&Device);
 				BindBufferMemoryInfo.memoryOffset = OutAllocation.Offset;
 
-				VERIFYVULKANRESULT(VulkanRHI::vkBindBufferMemory2(Device->GetHandle(), 1, &BindBufferMemoryInfo));
+				VERIFYVULKANRESULT(VulkanRHI::vkBindBufferMemory2(Device.GetHandle(), 1, &BindBufferMemoryInfo));
 			}
 
 			if (InDebugName)
 			{
-				VULKAN_SET_DEBUG_NAME((*Device), VK_OBJECT_TYPE_BUFFER, InBuffer, TEXT("%s"), InDebugName);
+				VULKAN_SET_DEBUG_NAME(Device, VK_OBJECT_TYPE_BUFFER, InBuffer, TEXT("%s"), InDebugName);
 			}
 		}
 		else
@@ -3109,7 +3109,7 @@ namespace VulkanRHI
 		ZeroVulkanStruct(MemoryReqs2, VK_STRUCTURE_TYPE_MEMORY_REQUIREMENTS_2_KHR);
 		MemoryReqs2.pNext = &DedMemoryReqs;
 
-		VulkanRHI::vkGetImageMemoryRequirements2(Device->GetHandle(), &ImageMemoryReqs2, &MemoryReqs2);
+		VulkanRHI::vkGetImageMemoryRequirements2(Device.GetHandle(), &ImageMemoryReqs2, &MemoryReqs2);
 
 		bool bUseDedicated = DedMemoryReqs.prefersDedicatedAllocation != VK_FALSE || DedMemoryReqs.requiresDedicatedAllocation != VK_FALSE;
 		if (bUseDedicated)
@@ -3153,7 +3153,7 @@ namespace VulkanRHI
 	void FMemoryManager::DumpMemory(bool bFullDump)
 	{
 		FScopeLock ScopeLock(&UsedFreeBufferAllocationsLock);
-		Device->GetDeviceMemoryManager().DumpMemory();
+		Device.GetDeviceMemoryManager().DumpMemory();
 		VULKAN_LOGMEMORY(TEXT("/******************************************* FMemoryManager ********************************************\\"));
 		VULKAN_LOGMEMORY(TEXT("HEAP DUMP"));
 
@@ -3365,8 +3365,8 @@ namespace VulkanRHI
 
 		FResourceHeapStats Staging;
 		TArray<FResourceHeapStats> DeviceHeaps;
-		Device->GetStagingManager().GetMemoryDump(Staging);
-		Device->GetDeviceMemoryManager().GetMemoryDump(DeviceHeaps);
+		Device.GetStagingManager().GetMemoryDump(Staging);
+		Device.GetDeviceMemoryManager().GetMemoryDump(DeviceHeaps);
 
 		VULKAN_LOGMEMORY(TEXT("SUMMARY"));
 		VULKAN_LOGMEMORY(TEXT("\t\tDevice Heaps                       |    Memory       Reserved      FreeMem       TotalMem   |  Allocs     -         |  Allocs              | Flags | Type   "));
@@ -4001,7 +4001,7 @@ namespace VulkanRHI
 	{
 		if (!MemoryAllocation->IsCoherent() || GForceCoherent != 0)
 		{
-			const VkDeviceSize NonCoherentAtomSize = Owner->GetParent()->GetLimits().nonCoherentAtomSize;
+			const VkDeviceSize NonCoherentAtomSize = Owner->GetDevice().GetLimits().nonCoherentAtomSize;
 			MemoryAllocation->FlushMappedMemory(Offset, AdjustToNonCoherentAtomSize(Offset, AllocationSize, MemoryAllocation->GetSize(), NonCoherentAtomSize));
 		}
 	}
@@ -4009,7 +4009,7 @@ namespace VulkanRHI
 	{
 		if (!MemoryAllocation->IsCoherent() || GForceCoherent != 0)
 		{
-			const VkDeviceSize NonCoherentAtomSize = Owner->GetParent()->GetLimits().nonCoherentAtomSize;
+			const VkDeviceSize NonCoherentAtomSize = Owner->GetDevice().GetLimits().nonCoherentAtomSize;
 			MemoryAllocation->InvalidateMappedMemory(Offset, AdjustToNonCoherentAtomSize(Offset, AllocationSize, MemoryAllocation->GetSize(), NonCoherentAtomSize));
 		}
 	}
@@ -4701,17 +4701,17 @@ namespace VulkanRHI
 	}
 
 
-	FGPUEvent::FGPUEvent(FVulkanDevice* InDevice)
-		: FDeviceChild(InDevice)
+	FGPUEvent::FGPUEvent(FVulkanDevice& InDevice)
+		: Device(InDevice)
 	{
 		VkEventCreateInfo Info;
 		ZeroVulkanStruct(Info, VK_STRUCTURE_TYPE_EVENT_CREATE_INFO);
-		VERIFYVULKANRESULT(VulkanRHI::vkCreateEvent(InDevice->GetHandle(), &Info, VULKAN_CPU_ALLOCATOR, &Handle));
+		VERIFYVULKANRESULT(VulkanRHI::vkCreateEvent(InDevice.GetHandle(), &Info, VULKAN_CPU_ALLOCATOR, &Handle));
 	}
 
 	FGPUEvent::~FGPUEvent()
 	{
-		Device->GetDeferredDeletionQueue().EnqueueResource(VulkanRHI::FDeferredDeletionQueue2::EType::Event, Handle);
+		Device.GetDeferredDeletionQueue().EnqueueResource(VulkanRHI::FDeferredDeletionQueue2::EType::Event, Handle);
 	}
 
 
@@ -4895,8 +4895,8 @@ namespace VulkanRHI
 
 
 
-	FTempBlockAllocator::FTempBlockAllocator(FVulkanDevice* InDevice, uint32 InBlockSize, uint32 InBlockAlignment, VkBufferUsageFlags InBufferUsage)
-		: FDeviceChild(InDevice)
+	FTempBlockAllocator::FTempBlockAllocator(FVulkanDevice& InDevice, uint32 InBlockSize, uint32 InBlockAlignment, VkBufferUsageFlags InBufferUsage)
+		: Device(InDevice)
 		, BlockSize(InBlockSize)
 		, BlockAlignment(InBlockAlignment)
 		, BufferUsage(InBufferUsage)
@@ -4906,10 +4906,10 @@ namespace VulkanRHI
 
 	FTempBlockAllocator::~FTempBlockAllocator()
 	{
-		auto FreeBlock = [](FVulkanDevice* InDevice, FTempMemoryBlock* Block)
+		auto FreeBlock = [](FVulkanDevice& InDevice, FTempMemoryBlock* Block)
 		{
-			VulkanRHI::vkDestroyBuffer(InDevice->GetHandle(), Block->Buffer, VULKAN_CPU_ALLOCATOR);
-			InDevice->GetMemoryManager().FreeVulkanAllocation(Block->Allocation, VulkanRHI::EVulkanFreeFlag_DontDefer);
+			VulkanRHI::vkDestroyBuffer(InDevice.GetHandle(), Block->Buffer, VULKAN_CPU_ALLOCATOR);
+			InDevice.GetMemoryManager().FreeVulkanAllocation(Block->Allocation, VulkanRHI::EVulkanFreeFlag_DontDefer);
 			delete Block;
 		};
 
@@ -4935,7 +4935,7 @@ namespace VulkanRHI
 
 		// Create buffer
 		{
-			NewBlock->Buffer = Device->CreateBuffer(BlockSize, BufferUsage);
+			NewBlock->Buffer = Device.CreateBuffer(BlockSize, BufferUsage);
 		}
 
 		// Allocate memory
@@ -4945,22 +4945,22 @@ namespace VulkanRHI
 				VulkanRHI::EVulkanAllocationFlags::PreferBAR |
 				VulkanRHI::EVulkanAllocationFlags::AutoBind |
 				VulkanRHI::EVulkanAllocationFlags::Dedicated;
-			Device->GetMemoryManager().AllocateBufferMemory(NewBlock->Allocation, NewBlock->Buffer, AllocFlags, TEXT("FTempBlockAllocator"), BlockAlignment);
+			Device.GetMemoryManager().AllocateBufferMemory(NewBlock->Allocation, NewBlock->Buffer, AllocFlags, TEXT("FTempBlockAllocator"), BlockAlignment);
 
 			// Pull the stat of the generic "buffer other" to put it in the temp block bucket (will get cleaned up soon)
 			DEC_DWORD_STAT_BY(STAT_VulkanAllocation_BufferOther, NewBlock->Allocation.Size);
 			INC_DWORD_STAT_BY(STAT_VulkanAllocation_TempBlocks, NewBlock->Allocation.Size);
 
-			NewBlock->MappedPointer = (uint8*)NewBlock->Allocation.GetMappedPointer(Device);
+			NewBlock->MappedPointer = (uint8*)NewBlock->Allocation.GetMappedPointer(&Device);
 		}
 
 		// Get the device addr if needed
-		if (VKHasAllFlags(BufferUsage, VK_BUFFER_USAGE_SHADER_DEVICE_ADDRESS_BIT) && Device->GetOptionalExtensions().HasBufferDeviceAddress)
+		if (VKHasAllFlags(BufferUsage, VK_BUFFER_USAGE_SHADER_DEVICE_ADDRESS_BIT) && Device.GetOptionalExtensions().HasBufferDeviceAddress)
 		{
 			VkBufferDeviceAddressInfo BufferInfo;
 			ZeroVulkanStruct(BufferInfo, VK_STRUCTURE_TYPE_BUFFER_DEVICE_ADDRESS_INFO);
 			BufferInfo.buffer = NewBlock->Buffer;
-			NewBlock->BufferAddress = VulkanRHI::vkGetBufferDeviceAddressKHR(Device->GetHandle(), &BufferInfo);
+			NewBlock->BufferAddress = VulkanRHI::vkGetBufferDeviceAddressKHR(Device.GetHandle(), &BufferInfo);
 		}
 
 		return NewBlock;
diff --git a/Engine/Source/Runtime/VulkanRHI/Private/VulkanMemory.h b/Engine/Source/Runtime/VulkanRHI/Private/VulkanMemory.h
--- a/Engine/Source/Runtime/VulkanRHI/Private/VulkanMemory.h
+++ b/Engine/Source/Runtime/VulkanRHI/Private/VulkanMemory.h
@@ -207,24 +207,6 @@ namespace VulkanRHI
 		FThreadSafeCounter NumRefs;
 	};
 
-	class FDeviceChild
-	{
-	public:
-		FDeviceChild(FVulkanDevice* InDevice)
-			: Device(InDevice)
-		{}
-
-		FVulkanDevice* GetParent() const
-		{
-			// Has to have one if we are asking for it...
-			check(Device);
-			return Device;
-		}
-
-	protected:
-		FVulkanDevice* const Device;
-	};
-
 	// An Allocation off a Device Heap. Lowest level of allocations and bounded by VkPhysicalDeviceLimits::maxMemoryAllocationCount.
 	class FDeviceMemoryAllocation
 	{
@@ -767,17 +749,22 @@ namespace VulkanRHI
 	ENUM_CLASS_FLAGS(EVulkanAllocationFlags);
 
 	// Manages heaps and their interactions
-	class FMemoryManager : public FDeviceChild
+	class FMemoryManager
 	{
 		friend class FVulkanAllocation;
 	public:
 
-		FMemoryManager(FVulkanDevice* InDevice);
+		FMemoryManager(FVulkanDevice& InDevice);
 		~FMemoryManager();
 
 		void Init();
 		void Deinit();
 
+		FVulkanDevice& GetDevice()
+		{
+			return Device;
+		}
+
 		static uint32 CalculateBufferAlignment(FVulkanDevice& InDevice, EBufferUsageFlags InUEUsage, bool bZeroSize);
 		static float CalculateBufferPriority(const VkBufferUsageFlags BufferUsageFlags);
 
@@ -814,6 +801,7 @@ namespace VulkanRHI
 		bool UpdateEvictThreshold(bool bLog);
 
 	protected:
+		FVulkanDevice& Device;
 		FDeviceMemoryManager* DeviceMemoryManager;
 		TArray<FVulkanResourceHeap*> ResourceTypeHeaps;
 
@@ -1093,10 +1081,10 @@ namespace VulkanRHI
 		void DestroyFence(FFence* Fence);
 	};
 
-	class FGPUEvent : public FDeviceChild, public FRefCount
+	class FGPUEvent : public FRefCount
 	{
 	public:
-		FGPUEvent(FVulkanDevice* InDevice);
+		FGPUEvent(FVulkanDevice& InDevice);
 		virtual ~FGPUEvent();
 
 		inline VkEvent GetHandle() const
@@ -1105,12 +1093,13 @@ namespace VulkanRHI
 		}
 
 	protected:
+		FVulkanDevice& Device;
 		VkEvent Handle;
 	};
 
 
 	// Simple temp allocation blocks used for volatile allocations
-	class FTempBlockAllocator : public FDeviceChild
+	class FTempBlockAllocator
 	{
 	public:
 		struct FTempMemoryBlock
@@ -1131,7 +1120,7 @@ namespace VulkanRHI
 			uint32 Offset = 0;
 		};
 
-		FTempBlockAllocator(FVulkanDevice* InDevice, uint32 InBlockSize, uint32 InBlockAlignment, VkBufferUsageFlags InBufferUsage);
+		FTempBlockAllocator(FVulkanDevice& InDevice, uint32 InBlockSize, uint32 InBlockAlignment, VkBufferUsageFlags InBufferUsage);
 		virtual ~FTempBlockAllocator();
 
 		uint8* Alloc(uint32 InSize, FVulkanContextCommon& Context, VkDescriptorBufferBindingInfoEXT& OutBindingInfo, VkDeviceSize& OutOffset);
@@ -1144,6 +1133,8 @@ namespace VulkanRHI
 		FTempMemoryBlock* AllocBlock();
 		FInternalAlloc InternalAlloc(uint32 InSize, FVulkanContextCommon& Context);
 
+		FVulkanDevice& Device;
+
 		const uint32 BlockSize;
 		const uint32 BlockAlignment;
 		const VkBufferUsageFlags BufferUsage;
diff --git a/Engine/Source/Runtime/VulkanRHI/Private/VulkanPipeline.cpp b/Engine/Source/Runtime/VulkanRHI/Private/VulkanPipeline.cpp
--- a/Engine/Source/Runtime/VulkanRHI/Private/VulkanPipeline.cpp
+++ b/Engine/Source/Runtime/VulkanRHI/Private/VulkanPipeline.cpp
@@ -1899,7 +1899,7 @@ FVulkanLayout* FVulkanPipelineStateCacheManager::FindOrAddLayout(const FVulkanDe
 		return *FoundLayout;
 	}
 
-	FVulkanLayout* Layout = new FVulkanLayout(Device, bGfxLayout, bUsesBindless);
+	FVulkanLayout* Layout = new FVulkanLayout(*Device, bGfxLayout, bUsesBindless);
 	Layout->DescriptorSetLayout.CopyFrom(DescriptorSetLayoutInfo);
 	Layout->Compile(DSetLayoutMap);
 
diff --git a/Engine/Source/Runtime/VulkanRHI/Private/VulkanRayTracing.cpp b/Engine/Source/Runtime/VulkanRHI/Private/VulkanRayTracing.cpp
--- a/Engine/Source/Runtime/VulkanRHI/Private/VulkanRayTracing.cpp
+++ b/Engine/Source/Runtime/VulkanRHI/Private/VulkanRayTracing.cpp
@@ -693,10 +693,9 @@ static VkGeometryInstanceFlagsKHR TranslateRayTracingInstanceFlags(ERayTracingIn
 	return Result;
 }
 
-FVulkanRayTracingScene::FVulkanRayTracingScene(FRayTracingSceneInitializer InInitializer, FVulkanDevice* InDevice)
-	: FDeviceChild(InDevice)
+FVulkanRayTracingScene::FVulkanRayTracingScene(FRayTracingSceneInitializer InInitializer, FVulkanDevice& InDevice)
+	: Device(InDevice)
 	, Initializer(MoveTemp(InInitializer))
-
 {
 	INC_DWORD_STAT(STAT_VulkanRayTracingAllocatedTLAS);
 
@@ -735,7 +734,7 @@ void FVulkanRayTracingScene::BindBuffer(FRHIBuffer* InBuffer, uint32 InBufferOff
 
 		check(InBufferOffset % GRHIRayTracingAccelerationStructureAlignment == 0);
 
-		View = MakeUnique<FVulkanView>(*Device, VK_DESCRIPTOR_TYPE_ACCELERATION_STRUCTURE_KHR);
+		View = MakeUnique<FVulkanView>(Device, VK_DESCRIPTOR_TYPE_ACCELERATION_STRUCTURE_KHR);
 		VkAccelerationStructureKHR NativeAccelerationStructureHandle = View->InitAsAccelerationStructureView(
 			AccelerationStructureBuffer
 			, InBufferOffset
@@ -746,7 +745,7 @@ void FVulkanRayTracingScene::BindBuffer(FRHIBuffer* InBuffer, uint32 InBufferOff
 
 		FString DebugNameString = Initializer.DebugName.ToString();
 		DebugNameString = (DebugNameString.IsEmpty()) ? TEXT("TLAS") : DebugNameString;
-		VULKAN_SET_DEBUG_NAME(*Device, VK_OBJECT_TYPE_ACCELERATION_STRUCTURE_KHR, NativeAccelerationStructureHandle, TEXT("%s"), *DebugNameString);
+		VULKAN_SET_DEBUG_NAME(Device, VK_OBJECT_TYPE_ACCELERATION_STRUCTURE_KHR, NativeAccelerationStructureHandle, TEXT("%s"), *DebugNameString);
 	}
 }
 
@@ -811,7 +810,7 @@ void BuildAccelerationStructure(
 	const VkDeviceAddress InstanceBufferAddress = InInstanceBuffer->GetDeviceAddress() + InInstanceOffset;
 
 	{
-		GetTLASBuildData(Scene.GetParent()->GetHandle(), NumInstances, InstanceBufferAddress, Scene.Initializer.BuildFlags, BuildMode, BuildData);
+		GetTLASBuildData(Scene.Device.GetHandle(), NumInstances, InstanceBufferAddress, Scene.Initializer.BuildFlags, BuildMode, BuildData);
 
 		BuildData.GeometryInfo.dstAccelerationStructure = Scene.View->GetAccelerationStructureView().Handle;
 		BuildData.GeometryInfo.srcAccelerationStructure = bIsUpdate ? Scene.View->GetAccelerationStructureView().Handle : nullptr;
@@ -847,13 +846,13 @@ void BuildAccelerationStructure(
 	Scene.bBuilt = true;
 }
 
-FVulkanRayTracingShaderTable::FVulkanRayTracingShaderTable(FRHICommandListBase& RHICmdList, FVulkanDevice* Device, const FRayTracingShaderBindingTableInitializer& InInitializer)
+FVulkanRayTracingShaderTable::FVulkanRayTracingShaderTable(FRHICommandListBase& RHICmdList, FVulkanDevice& InDevice, const FRayTracingShaderBindingTableInitializer& InInitializer)
 	: FRHIShaderBindingTable(InInitializer)
-	, FDeviceChild(Device)
+	, Device(InDevice)
 	, ShaderBindingMode(InInitializer.ShaderBindingMode)
 	, HitGroupIndexingMode(InInitializer.HitGroupIndexingMode)
-	, HandleSize(Device->GetOptionalExtensionProperties().RayTracingPipelineProps.shaderGroupHandleSize)
-	, HandleSizeAligned(Align(HandleSize, Device->GetOptionalExtensionProperties().RayTracingPipelineProps.shaderGroupHandleAlignment))
+	, HandleSize(Device.GetOptionalExtensionProperties().RayTracingPipelineProps.shaderGroupHandleSize)
+	, HandleSizeAligned(Align(HandleSize, Device.GetOptionalExtensionProperties().RayTracingPipelineProps.shaderGroupHandleAlignment))
 {
 	check(ShaderBindingMode != ERayTracingShaderBindingMode::Disabled);
 
@@ -903,17 +902,17 @@ void FVulkanRayTracingShaderTable::ReleaseLocalBuffers()
 	ReleaseLocalBuffer(Device, Callable);
 }
 
-void FVulkanRayTracingShaderTable::ReleaseLocalBuffer(FVulkanDevice* Device, FVulkanShaderTableAllocation& Alloc)
+void FVulkanRayTracingShaderTable::ReleaseLocalBuffer(FVulkanDevice& Device, FVulkanShaderTableAllocation& Alloc)
 {
 	if (Alloc.LocalBuffer != VK_NULL_HANDLE)
 	{
-		Device->GetDeferredDeletionQueue().EnqueueResource(VulkanRHI::FDeferredDeletionQueue2::EType::Buffer, Alloc.LocalBuffer);
+		Device.GetDeferredDeletionQueue().EnqueueResource(VulkanRHI::FDeferredDeletionQueue2::EType::Buffer, Alloc.LocalBuffer);
 		Alloc.LocalBuffer = VK_NULL_HANDLE;
 	}
 
 	if (Alloc.LocalAllocation.IsValid())
 	{
-		Device->GetMemoryManager().FreeVulkanAllocation(Alloc.LocalAllocation);
+		Device.GetMemoryManager().FreeVulkanAllocation(Alloc.LocalAllocation);
 	}
 
 	Alloc.Region.deviceAddress = 0;
@@ -958,7 +957,7 @@ void FVulkanRayTracingShaderTable::SetSlot(EShaderFrequency Frequency, uint32 Ds
 VkStridedDeviceAddressRegionKHR FVulkanRayTracingShaderTable::CommitRayGenShader(FVulkanCommandListContext& Context, uint32 SrcHandleIndex, TConstArrayView<uint8> SrcHandleData)
 {
 	VkStridedDeviceAddressRegionKHR RayGenRegion;
-	uint8* pMappedMemory = Device->GetTempBlockAllocator().Alloc(HandleSize, Context, &RayGenRegion);
+	uint8* pMappedMemory = Device.GetTempBlockAllocator().Alloc(HandleSize, Context, &RayGenRegion);
 
 	FMemory::Memcpy(pMappedMemory, &SrcHandleData[SrcHandleIndex * HandleSize], HandleSize);
 
@@ -1015,7 +1014,7 @@ void FVulkanRayTracingShaderTable::Commit(FVulkanCommandListContext& Context, FR
 	VkMemoryBarrier BarrierBefore = { VK_STRUCTURE_TYPE_MEMORY_BARRIER, nullptr, VK_ACCESS_MEMORY_READ_BIT | VK_ACCESS_MEMORY_WRITE_BIT, VK_ACCESS_TRANSFER_READ_BIT | VK_ACCESS_TRANSFER_WRITE_BIT };
 	VulkanRHI::vkCmdPipelineBarrier(CommandBuffer.GetHandle(), VK_PIPELINE_STAGE_BOTTOM_OF_PIPE_BIT, VK_PIPELINE_STAGE_TRANSFER_BIT, 0, 1, &BarrierBefore, 0, nullptr, 0, nullptr);
 
-	auto CommitBuffer = [Device = Device, &Context](FVulkanShaderTableAllocation& Alloc)
+	auto CommitBuffer = [&Device = Device, &Context](FVulkanShaderTableAllocation& Alloc)
 	{
 		if (Alloc.bIsDirty)
 		{
@@ -1023,20 +1022,20 @@ void FVulkanRayTracingShaderTable::Commit(FVulkanCommandListContext& Context, FR
 			{
 				ReleaseLocalBuffer(Device, Alloc);
 
-				const VkDevice DeviceHandle = Device->GetHandle();
-				const VkPhysicalDeviceRayTracingPipelinePropertiesKHR& RayTracingPipelineProps = Device->GetOptionalExtensionProperties().RayTracingPipelineProps;
+				const VkDevice DeviceHandle = Device.GetHandle();
+				const VkPhysicalDeviceRayTracingPipelinePropertiesKHR& RayTracingPipelineProps = Device.GetOptionalExtensionProperties().RayTracingPipelineProps;
 
 				// Fetch staging buffer and fill it
-				VulkanRHI::FStagingBuffer* StagingBuffer = Device->GetStagingManager().AcquireBuffer(Alloc.Region.size);
+				VulkanRHI::FStagingBuffer* StagingBuffer = Device.GetStagingManager().AcquireBuffer(Alloc.Region.size);
 				FMemory::Memcpy(StagingBuffer->GetMappedPointer(), Alloc.HostBuffer.GetData(), Alloc.Region.size);
 
 				// Alloc a new Local buffer
 				{
 					const VkBufferUsageFlags BufferUsageFlags = VK_BUFFER_USAGE_TRANSFER_DST_BIT | VK_BUFFER_USAGE_SHADER_BINDING_TABLE_BIT_KHR | VK_BUFFER_USAGE_SHADER_DEVICE_ADDRESS_BIT;
-					Alloc.LocalBuffer = Device->CreateBuffer(Alloc.Region.size, BufferUsageFlags);
+					Alloc.LocalBuffer = Device.CreateBuffer(Alloc.Region.size, BufferUsageFlags);
 
 					const VulkanRHI::EVulkanAllocationFlags AllocFlags = VulkanRHI::EVulkanAllocationFlags::AutoBind | VulkanRHI::EVulkanAllocationFlags::Dedicated;
-					Device->GetMemoryManager().AllocateBufferMemory(Alloc.LocalAllocation, Alloc.LocalBuffer, AllocFlags, TEXT("LocalShaderTableAllocation"), RayTracingPipelineProps.shaderGroupBaseAlignment);
+					Device.GetMemoryManager().AllocateBufferMemory(Alloc.LocalAllocation, Alloc.LocalBuffer, AllocFlags, TEXT("LocalShaderTableAllocation"), RayTracingPipelineProps.shaderGroupBaseAlignment);
 
 					VkBufferDeviceAddressInfoKHR DeviceAddressInfo;
 					ZeroVulkanStruct(DeviceAddressInfo, VK_STRUCTURE_TYPE_BUFFER_DEVICE_ADDRESS_INFO);
@@ -1050,7 +1049,7 @@ void FVulkanRayTracingShaderTable::Commit(FVulkanCommandListContext& Context, FR
 				RegionInfo.size = Alloc.Region.size;
 				VulkanRHI::vkCmdCopyBuffer(Context.GetCommandBuffer().GetHandle(), StagingBuffer->GetHandle(), Alloc.LocalBuffer, 1, &RegionInfo);
 
-				Device->GetStagingManager().ReleaseBuffer(&Context, StagingBuffer);
+				Device.GetStagingManager().ReleaseBuffer(&Context, StagingBuffer);
 			}
 			else
 			{
@@ -1121,7 +1120,7 @@ FRayTracingAccelerationStructureSize FVulkanDynamicRHI::RHICalcRayTracingGeometr
 
 FRayTracingSceneRHIRef FVulkanDynamicRHI::RHICreateRayTracingScene(FRayTracingSceneInitializer Initializer)
 {
-	return new FVulkanRayTracingScene(MoveTemp(Initializer), GetDevice());
+	return new FVulkanRayTracingScene(MoveTemp(Initializer), *Device);
 }
 
 FRayTracingGeometryRHIRef FVulkanDynamicRHI::RHICreateRayTracingGeometry(FRHICommandListBase& RHICmdList, const FRayTracingGeometryInitializer& Initializer)
@@ -1131,12 +1130,12 @@ FRayTracingGeometryRHIRef FVulkanDynamicRHI::RHICreateRayTracingGeometry(FRHICom
 
 FRayTracingPipelineStateRHIRef FVulkanDynamicRHI::RHICreateRayTracingPipelineState(const FRayTracingPipelineStateInitializer& Initializer)
 {
-	return new FVulkanRayTracingPipelineState(GetDevice(), Initializer);
+	return new FVulkanRayTracingPipelineState(*Device, Initializer);
 }
 
 FShaderBindingTableRHIRef FVulkanDynamicRHI::RHICreateShaderBindingTable(FRHICommandListBase& RHICmdList, const FRayTracingShaderBindingTableInitializer& Initializer)
 {
-	return new FVulkanRayTracingShaderTable(RHICmdList, GetDevice(), Initializer);
+	return new FVulkanRayTracingShaderTable(RHICmdList, *Device, Initializer);
 }
 
 void FVulkanCommandListContext::RHIClearShaderBindingTable(FRHIShaderBindingTable* InSBT)
@@ -1510,7 +1509,7 @@ void FVulkanDevice::InitializeRayTracing()
 // NOTE: Should currently only be used by InstanceCulling due to a binding that isn't stripped by DXC. See also USE_INSTANCE_CULLING_DATA for same issue in CS.
 static FRWLock DummyUBLock;
 static TMap<uint32, FUniformBufferRHIRef> DummyUBs;
-static FVulkanUniformBuffer* GetDummyUB(FVulkanDevice* Device, uint32 UBLayoutHash)
+static FVulkanUniformBuffer* GetDummyUB(FVulkanDevice& Device, uint32 UBLayoutHash)
 {
 	{
 		FRWScopeLock ScopedReadLock(DummyUBLock, SLT_ReadOnly);
@@ -1528,7 +1527,7 @@ static FVulkanUniformBuffer* GetDummyUB(FVulkanDevice* Device, uint32 UBLayoutHa
 		const FRHIUniformBufferLayout* DummyLayout = DummyMetadata->GetLayoutPtr();
 		TArray<uint8> DummyContent;
 		DummyContent.SetNumZeroed(DummyLayout->ConstantBufferSize);
-		FVulkanUniformBuffer* DummyUB = new FVulkanUniformBuffer(*Device, DummyLayout, DummyContent.GetData(), UniformBuffer_MultiFrame, EUniformBufferValidation::None);
+		FVulkanUniformBuffer* DummyUB = new FVulkanUniformBuffer(Device, DummyLayout, DummyContent.GetData(), UniformBuffer_MultiFrame, EUniformBufferValidation::None);
 		DummyUBs.Add(UBLayoutHash, DummyUB);
 		const FString& LayoutName = DummyLayout->GetDebugName();
 		UE_LOG(LogRHI, Warning, TEXT("Vulkan ray tracing using DummyUB for %s."), LayoutName.IsEmpty() ? TEXT("<unknown>") : *LayoutName);
@@ -1549,22 +1548,23 @@ void FVulkanDevice::CleanUpRayTracing()
 }
 
 
-FVulkanRayTracingPipelineState::FVulkanRayTracingPipelineState(FVulkanDevice* const InDevice, const FRayTracingPipelineStateInitializer& Initializer)
-	: FRHIRayTracingPipelineState(Initializer), FDeviceChild(InDevice)
+FVulkanRayTracingPipelineState::FVulkanRayTracingPipelineState(FVulkanDevice& InDevice, const FRayTracingPipelineStateInitializer& Initializer)
+	: FRHIRayTracingPipelineState(Initializer)
+	, Device(InDevice)
 	, bIsPartialPipeline(Initializer.bPartial)
 {
-	checkf(InDevice->SupportsBindless(), TEXT("Vulkan ray tracing pipelines are only supported in bindless."));
-	checkf(Initializer.MaxAttributeSizeInBytes <= InDevice->GetOptionalExtensionProperties().RayTracingPipelineProps.maxRayHitAttributeSize, 
+	checkf(InDevice.SupportsBindless(), TEXT("Vulkan ray tracing pipelines are only supported in bindless."));
+	checkf(Initializer.MaxAttributeSizeInBytes <= InDevice.GetOptionalExtensionProperties().RayTracingPipelineProps.maxRayHitAttributeSize,
 		TEXT("Required attribute size (%u) too large for current device (%u)."), 
 		Initializer.MaxAttributeSizeInBytes, 
-		InDevice->GetOptionalExtensionProperties().RayTracingPipelineProps.maxRayHitAttributeSize);
+		InDevice.GetOptionalExtensionProperties().RayTracingPipelineProps.maxRayHitAttributeSize);
 
 	TArrayView<FRHIRayTracingShader*> InitializerRayGenShaders = Initializer.GetRayGenTable();
 	TArrayView<FRHIRayTracingShader*> InitializerMissShaders = Initializer.GetMissTable();
 	TArrayView<FRHIRayTracingShader*> InitializerHitGroupShaders = Initializer.GetHitGroupTable();
 	TArrayView<FRHIRayTracingShader*> InitializerCallableShaders = Initializer.GetCallableTable();
 
-	FVulkanRayTracingPipelineLibraryCache* LibraryCache = InDevice->GetRayTracingPipelineLibraryCache();
+	FVulkanRayTracingPipelineLibraryCache* LibraryCache = InDevice.GetRayTracingPipelineLibraryCache();
 
 	TArray<FVulkanRayTracingPipelineLibraryCache::FPipelineLibrary*> PipelineLibraries;
 	FGraphEventArray PendingLibraryTasks;
@@ -1629,18 +1629,18 @@ FVulkanRayTracingPipelineState::FVulkanRayTracingPipelineState(FVulkanDevice* co
 	RayTracingPipelineCreateInfo.pLibraryInfo = &PipelineLibraryCreateInfo;
 	RayTracingPipelineCreateInfo.pLibraryInterface = &RayTracingPipelineInterfaceCreateInfo;
 	RayTracingPipelineCreateInfo.maxPipelineRayRecursionDepth = 1;
-	RayTracingPipelineCreateInfo.layout = InDevice->GetBindlessDescriptorManager()->GetPipelineLayout();
+	RayTracingPipelineCreateInfo.layout = InDevice.GetBindlessDescriptorManager()->GetPipelineLayout();
 	RayTracingPipelineCreateInfo.flags = VK_PIPELINE_CREATE_DESCRIPTOR_BUFFER_BIT_EXT;
 
-	VERIFYVULKANRESULT_EXPANDED(InDevice->GetPipelineStateCache()->CreateRayTracingPipeline(RayTracingPipelineCreateInfo, bIsPartialPipeline, Pipeline));
+	VERIFYVULKANRESULT_EXPANDED(InDevice.GetPipelineStateCache()->CreateRayTracingPipeline(RayTracingPipelineCreateInfo, bIsPartialPipeline, Pipeline));
 
 	// Grab all shader handles for each stage
 	{
-		const VkPhysicalDeviceRayTracingPipelinePropertiesKHR& RayTracingPipelineProps = InDevice->GetOptionalExtensionProperties().RayTracingPipelineProps;
+		const VkPhysicalDeviceRayTracingPipelinePropertiesKHR& RayTracingPipelineProps = InDevice.GetOptionalExtensionProperties().RayTracingPipelineProps;
 		const uint32 HandleSize = RayTracingPipelineProps.shaderGroupHandleSize;
 
 		uint32 HandleOffset = 0;
-		auto FetchShaderHandles = [&HandleOffset, HandleSize, InDevice](VkPipeline RTPipeline, uint32 HandleCount) 
+		auto FetchShaderHandles = [&HandleOffset, HandleSize, &InDevice](VkPipeline RTPipeline, uint32 HandleCount) 
 		{
 			TArray<uint8> OutHandleStorage;
 
@@ -1649,7 +1649,7 @@ FVulkanRayTracingPipelineState::FVulkanRayTracingPipelineState(FVulkanDevice* co
 				const uint32 ShaderHandleStorageSize = HandleCount * HandleSize;
 				OutHandleStorage.AddUninitialized(ShaderHandleStorageSize);
 
-				VERIFYVULKANRESULT(VulkanDynamicAPI::vkGetRayTracingShaderGroupHandlesKHR(InDevice->GetHandle(), RTPipeline, HandleOffset, HandleCount, ShaderHandleStorageSize, OutHandleStorage.GetData()));
+				VERIFYVULKANRESULT(VulkanDynamicAPI::vkGetRayTracingShaderGroupHandlesKHR(InDevice.GetHandle(), RTPipeline, HandleOffset, HandleCount, ShaderHandleStorageSize, OutHandleStorage.GetData()));
 
 				HandleOffset += HandleCount;
 			}
@@ -1671,7 +1671,7 @@ FVulkanRayTracingPipelineState::~FVulkanRayTracingPipelineState()
 {
 	if (Pipeline != VK_NULL_HANDLE)
 	{
-		VulkanRHI::vkDestroyPipeline(Device->GetHandle(), Pipeline, VULKAN_CPU_ALLOCATOR);
+		VulkanRHI::vkDestroyPipeline(Device.GetHandle(), Pipeline, VULKAN_CPU_ALLOCATOR);
 		Pipeline = VK_NULL_HANDLE;
 	}
 }
@@ -1764,10 +1764,10 @@ bool FVulkanRayTracingCompactedSizeQueryPool::TryGetResults(uint32 NumResults)
 	return false;
 }
 
-FVulkanRayTracingCompactionRequestHandler::FVulkanRayTracingCompactionRequestHandler(FVulkanDevice* const InDevice)
-	: VulkanRHI::FDeviceChild(InDevice)
+FVulkanRayTracingCompactionRequestHandler::FVulkanRayTracingCompactionRequestHandler(FVulkanDevice& InDevice)
+	: Device(InDevice)
 {
-	QueryPool = new FVulkanRayTracingCompactedSizeQueryPool(*InDevice, GVulkanRayTracingMaxBatchedCompaction);
+	QueryPool = new FVulkanRayTracingCompactedSizeQueryPool(InDevice, GVulkanRayTracingMaxBatchedCompaction);
 
 	ActiveRequests.Reserve(GVulkanRayTracingMaxBatchedCompaction);
 	ActiveBLASes.Reserve(GVulkanRayTracingMaxBatchedCompaction);
@@ -2030,7 +2030,7 @@ static void SetSystemParametersUB(FVulkanHitGroupSystemParameters& OutSystemPara
 		// :todo-jn: Hack to force in a DummyCullingBuffer in cases where it should have been culled from source (see SPIRV-Tools Issue 4902).
 		if (!UniformBuffer)
 		{
-			UniformBuffer = GetDummyUB(ShaderTable->GetParent(), UniformBufferInfo.LayoutHash);
+			UniformBuffer = GetDummyUB(ShaderTable->GetDevice(), UniformBufferInfo.LayoutHash);
 		}
 
 		check(UniformBuffer);
diff --git a/Engine/Source/Runtime/VulkanRHI/Private/VulkanRayTracing.h b/Engine/Source/Runtime/VulkanRHI/Private/VulkanRayTracing.h
--- a/Engine/Source/Runtime/VulkanRHI/Private/VulkanRayTracing.h
+++ b/Engine/Source/Runtime/VulkanRHI/Private/VulkanRayTracing.h
@@ -65,10 +65,10 @@ struct FVkRtBLASBuildData
 	VkAccelerationStructureBuildSizesInfoKHR SizesInfo;
 };
 
-class FVulkanRayTracingShaderTable : public FRHIShaderBindingTable, public VulkanRHI::FDeviceChild
+class FVulkanRayTracingShaderTable : public FRHIShaderBindingTable
 {
 public:
-	FVulkanRayTracingShaderTable(FRHICommandListBase& RHICmdList, FVulkanDevice* Device, const FRayTracingShaderBindingTableInitializer& InInitializer);
+	FVulkanRayTracingShaderTable(FRHICommandListBase& RHICmdList, FVulkanDevice& InDevice, const FRayTracingShaderBindingTableInitializer& InInitializer);
 	~FVulkanRayTracingShaderTable();
 
 	void ReleaseLocalBuffers();
@@ -123,6 +123,11 @@ public:
 		return ShaderBindingMode;
 	}
 
+	FVulkanDevice& GetDevice()
+	{
+		return Device;
+	}
+
 	// Ray tracing shader bindings can be processed in parallel.
 	// Each concurrent worker gets its own dedicated descriptor cache instance to avoid contention or locking.
 	// Scaling beyond 5 total threads does not yield any speedup in practice (RHI thread + 4 parallel workers).
@@ -152,8 +157,10 @@ private:
 	};
 
 	FVulkanShaderTableAllocation& GetAlloc(EShaderFrequency Frequency);
-	static void ReleaseLocalBuffer(FVulkanDevice* Device, FVulkanShaderTableAllocation& Alloc);
+	static void ReleaseLocalBuffer(FVulkanDevice& Device, FVulkanShaderTableAllocation& Alloc);
 	
+	FVulkanDevice& Device;
+
 	ERayTracingShaderBindingMode ShaderBindingMode = ERayTracingShaderBindingMode::Disabled;
 	ERayTracingHitGroupIndexingMode HitGroupIndexingMode = ERayTracingHitGroupIndexingMode::Allow;
 
@@ -211,12 +218,12 @@ public:
 	FName OwnerName;		// Store the path name of the owner object for resource tracking
 };
 
-class FVulkanRayTracingScene : public FRHIRayTracingScene, public VulkanRHI::FDeviceChild
+class FVulkanRayTracingScene : public FRHIRayTracingScene
 {
 	friend FVulkanCommandListContext;
 
 public:
-	FVulkanRayTracingScene(FRayTracingSceneInitializer Initializer, FVulkanDevice* InDevice);
+	FVulkanRayTracingScene(FRayTracingSceneInitializer Initializer, FVulkanDevice& InDevice);
 	~FVulkanRayTracingScene();
 
 	const FRayTracingSceneInitializer& GetInitializer() const override final { return Initializer; }
@@ -230,6 +237,8 @@ public:
 
 	using FRHIRayTracingAccelerationStructure::SizeInfo;
 
+	FVulkanDevice& Device;
+
 	const FRayTracingSceneInitializer Initializer;
 
 	// Unique list of geometries referenced by all instances in this scene.
@@ -257,12 +266,12 @@ private:
 };
 
 
-class FVulkanRayTracingPipelineState : public FRHIRayTracingPipelineState, public VulkanRHI::FDeviceChild
+class FVulkanRayTracingPipelineState : public FRHIRayTracingPipelineState
 {
 public:
 
 	UE_NONCOPYABLE(FVulkanRayTracingPipelineState);
-	FVulkanRayTracingPipelineState(FVulkanDevice* const InDevice, const FRayTracingPipelineStateInitializer& Initializer);
+	FVulkanRayTracingPipelineState(FVulkanDevice& InDevice, const FRayTracingPipelineStateInitializer& Initializer);
 	~FVulkanRayTracingPipelineState();
 
 	VkPipeline GetPipeline() const
@@ -281,6 +290,7 @@ public:
 	const TArray<uint8>& GetShaderHandles(EShaderFrequency Frequency) const;
 
 private:
+	FVulkanDevice& Device;
 
 	struct ShaderData
 	{
@@ -320,12 +330,12 @@ public:
 };
 
 // Manages all the pending BLAS compaction requests
-class FVulkanRayTracingCompactionRequestHandler : public VulkanRHI::FDeviceChild
+class FVulkanRayTracingCompactionRequestHandler
 {
 public:
 	UE_NONCOPYABLE(FVulkanRayTracingCompactionRequestHandler)
 
-	FVulkanRayTracingCompactionRequestHandler(FVulkanDevice* const InDevice);
+	FVulkanRayTracingCompactionRequestHandler(FVulkanDevice& InDevice);
 	~FVulkanRayTracingCompactionRequestHandler()
 	{
 		check(PendingRequests.IsEmpty());
@@ -338,6 +348,7 @@ public:
 	void Update(FVulkanCommandListContext& InCommandContext);
 
 private:
+	FVulkanDevice& Device;
 
 	FCriticalSection CS;
 	TArray<FVulkanRayTracingGeometry*> PendingRequests;
diff --git a/Engine/Source/Runtime/VulkanRHI/Private/VulkanViewport.cpp b/Engine/Source/Runtime/VulkanRHI/Private/VulkanViewport.cpp
--- a/Engine/Source/Runtime/VulkanRHI/Private/VulkanViewport.cpp
+++ b/Engine/Source/Runtime/VulkanRHI/Private/VulkanViewport.cpp
@@ -99,8 +99,8 @@ FVulkanBackBuffer::~FVulkanBackBuffer()
 	ReleaseAcquiredImage();
 }
 
-FVulkanViewport::FVulkanViewport(FVulkanDevice* InDevice, void* InWindowHandle, uint32 InSizeX, uint32 InSizeY, bool bInIsFullscreen, EPixelFormat InPreferredPixelFormat)
-	: VulkanRHI::FDeviceChild(InDevice)
+FVulkanViewport::FVulkanViewport(FVulkanDevice& InDevice, void* InWindowHandle, uint32 InSizeX, uint32 InSizeY, bool bInIsFullscreen, EPixelFormat InPreferredPixelFormat)
+	: Device(InDevice)
 	, SizeX(InSizeX)
 	, SizeY(InSizeY)
 	, bIsFullscreen(bInIsFullscreen)
@@ -162,7 +162,7 @@ FVulkanViewport::~FVulkanViewport()
 			RenderingDoneSemaphores[Index]->Release();
 
 			// FIXME: race condition on TransitionAndLayoutManager, could this be called from RT while RHIT is active?
-			Device->NotifyDeletedImage(BackBufferImages[Index]->Image, true);
+			Device.NotifyDeletedImage(BackBufferImages[Index]->Image, true);
 			BackBufferImages[Index] = nullptr;
 		}
 
@@ -283,7 +283,7 @@ void FVulkanViewport::IssueFrameEvent()
 {
 	if (FVulkanPlatform::RequiresWaitingForFrameCompletionEvent())
 	{
-		FVulkanCommandListContextImmediate& ImmediateContext = Device->GetImmediateContext();
+		FVulkanCommandListContextImmediate& ImmediateContext = Device.GetImmediateContext();
 		LastFrameSyncPoint = ImmediateContext.GetContextSyncPoint();
 		ImmediateContext.FlushCommands();
 	}
@@ -381,7 +381,7 @@ void FVulkanViewport::InitImages(FVulkanContextCommon& Context, TConstArrayView<
 		ViewDesc.NumMips = 1;
 		ViewDesc.ArraySliceIndex = 0;
 		ViewDesc.NumArraySlices = 1;
-		TextureViews.Add((new FVulkanView(*Device, VK_DESCRIPTOR_TYPE_SAMPLED_IMAGE))->InitAsTextureView(Images[Index], ViewDesc));
+		TextureViews.Add((new FVulkanView(Device, VK_DESCRIPTOR_TYPE_SAMPLED_IMAGE))->InitAsTextureView(Images[Index], ViewDesc));
 
 		// Clear the swapchain to avoid a validation warning, and transition to PresentSrc
 		{
@@ -391,9 +391,9 @@ void FVulkanViewport::InitImages(FVulkanContextCommon& Context, TConstArrayView<
 		}
 
 #if VULKAN_ENABLE_DRAW_MARKERS
-		if (Device->GetSetDebugName())
+		if (Device.GetSetDebugName())
 		{
-			VulkanRHI::SetDebugName(Device->GetSetDebugName(), Device->GetHandle(), BackBufferImages[Index]->Image, "VulkanBackBuffer");
+			VulkanRHI::SetDebugName(Device.GetSetDebugName(), Device.GetHandle(), BackBufferImages[Index]->Image, "VulkanBackBuffer");
 		}
 #endif
 	}
@@ -414,7 +414,7 @@ void FVulkanViewport::CreateSwapchain(FVulkanCommandListContext& Context, FVulka
 			uint32 DesiredNumBackBuffers = NUM_BUFFERS;
 			TArray<VkImage> Images;
 			SwapChain = new FVulkanSwapChain(
-				FVulkanDynamicRHI::Get().Instance, *Device,
+				FVulkanDynamicRHI::Get().Instance, Device,
 				PixelFormat, SizeX, SizeY, bIsFullscreen,
 				&DesiredNumBackBuffers,
 				Images,
@@ -438,13 +438,13 @@ void FVulkanViewport::CreateSwapchain(FVulkanCommandListContext& Context, FVulka
 			{
 				for (int32 Index = 0, NumBuffers = RenderingDoneSemaphores.Num(); Index < NumBuffers; ++Index)
 				{
-					RenderingDoneSemaphores[Index] = new VulkanRHI::FSemaphore(*Device);
+					RenderingDoneSemaphores[Index] = new VulkanRHI::FSemaphore(Device);
 					RenderingDoneSemaphores[Index]->AddRef();
 				}
 			}
 		}
 
-		RHIBackBuffer = new FVulkanBackBuffer(*Device, this, PixelFormat, SizeX, SizeY, TexCreate_RenderTargetable | TexCreate_ShaderResource | TexCreate_ResolveTargetable);
+		RHIBackBuffer = new FVulkanBackBuffer(Device, this, PixelFormat, SizeX, SizeY, TexCreate_RenderTargetable | TexCreate_ShaderResource | TexCreate_ResolveTargetable);
 	}
 	else
 	{
@@ -453,7 +453,7 @@ void FVulkanViewport::CreateSwapchain(FVulkanCommandListContext& Context, FVulka
 		{
 			if (RecreateInfo->SwapChain)
 			{
-				FVulkanPlatform::DestroySwapchainKHR(Device->GetHandle(), RecreateInfo->SwapChain, VULKAN_CPU_ALLOCATOR);
+				FVulkanPlatform::DestroySwapchainKHR(Device.GetHandle(), RecreateInfo->SwapChain, VULKAN_CPU_ALLOCATOR);
 				RecreateInfo->SwapChain = VK_NULL_HANDLE;
 			}
 			if (RecreateInfo->Surface)
@@ -489,9 +489,9 @@ void FVulkanViewport::CreateSwapchain(FVulkanCommandListContext& Context, FVulka
 		Barrier.Execute(&Context.GetCommandBuffer());
 
 #if VULKAN_ENABLE_DRAW_MARKERS
-		if (Device->GetSetDebugName())
+		if (Device.GetSetDebugName())
 		{
-			VulkanRHI::SetDebugName(Device->GetSetDebugName(), Device->GetHandle(), RenderingBackBuffer->Image, "RenderingBackBuffer");
+			VulkanRHI::SetDebugName(Device.GetSetDebugName(), Device.GetHandle(), RenderingBackBuffer->Image, "RenderingBackBuffer");
 		}
 #endif
 	}
@@ -517,17 +517,17 @@ void FVulkanViewport::DestroySwapchain(FVulkanSwapChainRecreateInfo* RecreateInf
 		TextureViews.Empty();
 		for (int32 Index = 0, NumBuffers = BackBufferImages.Num(); Index < NumBuffers; ++Index)
 		{
-			Device->NotifyDeletedImage(BackBufferImages[Index]->Image, true);
+			Device.NotifyDeletedImage(BackBufferImages[Index]->Image, true);
 			BackBufferImages[Index] = nullptr;
 		}
 		
-		Device->GetDeferredDeletionQueue().ReleaseResources(true);
+		Device.GetDeferredDeletionQueue().ReleaseResources(true);
 
 		SwapChain->Destroy(RecreateInfo);
 		delete SwapChain;
 		SwapChain = nullptr;
 
-		Device->GetDeferredDeletionQueue().ReleaseResources(true);
+		Device.GetDeferredDeletionQueue().ReleaseResources(true);
 	}
 
 	AcquiredImageIndex = -1;
@@ -858,7 +858,7 @@ FViewportRHIRef FVulkanDynamicRHI::RHICreateViewport(void* WindowHandle, uint32
 		PreferredPixelFormat = EDefaultBackBufferPixelFormat::Convert2PixelFormat(EDefaultBackBufferPixelFormat::FromInt(CVarDefaultBackBufferPixelFormat->GetValueOnAnyThread()));
 	}
 
-	return new FVulkanViewport(Device, WindowHandle, SizeX, SizeY, bIsFullscreen, PreferredPixelFormat);
+	return new FVulkanViewport(*Device, WindowHandle, SizeX, SizeY, bIsFullscreen, PreferredPixelFormat);
 }
 
 void FVulkanDynamicRHI::RHIResizeViewport(FRHIViewport* ViewportRHI, uint32 SizeX, uint32 SizeY, bool bIsFullscreen, EPixelFormat PreferredPixelFormat)
diff --git a/Engine/Source/Runtime/VulkanRHI/Private/VulkanViewport.h b/Engine/Source/Runtime/VulkanRHI/Private/VulkanViewport.h
--- a/Engine/Source/Runtime/VulkanRHI/Private/VulkanViewport.h
+++ b/Engine/Source/Runtime/VulkanRHI/Private/VulkanViewport.h
@@ -39,12 +39,12 @@ private:
 };
 
 
-class FVulkanViewport : public FRHIViewport, public VulkanRHI::FDeviceChild
+class FVulkanViewport : public FRHIViewport
 {
 public:
 	enum { NUM_BUFFERS = 3 };
 
-	FVulkanViewport(FVulkanDevice* InDevice, void* InWindowHandle, uint32 InSizeX,uint32 InSizeY,bool bInIsFullscreen, EPixelFormat InPreferredPixelFormat);
+	FVulkanViewport(FVulkanDevice& InDevice, void* InWindowHandle, uint32 InSizeX,uint32 InSizeY,bool bInIsFullscreen, EPixelFormat InPreferredPixelFormat);
 	~FVulkanViewport();
 
 	FTextureRHIRef GetBackBuffer(FRHICommandListImmediate& RHICmdList);
@@ -112,6 +112,8 @@ public:
 	void RecreateSwapchain(FVulkanCommandListContext& Context, FVulkanPlatformWindowContext& WindowContext);
 
 protected:
+	FVulkanDevice& Device;
+
 	// NUM_BUFFERS don't have to match exactly as the driver can require a minimum number larger than NUM_BUFFERS. Provide some slack
 	TArray<TRefCountPtr<FVulkanTexture>, TInlineAllocator<NUM_BUFFERS*2>> BackBufferImages;
 	TArray<VulkanRHI::FSemaphore*, TInlineAllocator<NUM_BUFFERS*2>> RenderingDoneSemaphores;
