VulkanRHI: Clean up: Add missing line endings in cvar descriptions.

#rnx

diff --git a/Engine/Source/Runtime/VulkanRHI/Private/VulkanPipeline.cpp b/Engine/Source/Runtime/VulkanRHI/Private/VulkanPipeline.cpp
--- a/Engine/Source/Runtime/VulkanRHI/Private/VulkanPipeline.cpp
+++ b/Engine/Source/Runtime/VulkanRHI/Private/VulkanPipeline.cpp
@@ -55,7 +55,7 @@ static TAtomic<uint64> SPipelineGfxCount;
 
 static const double HitchTime = 1.0 / 1000.0;
 
-TAutoConsoleVariable<int32> CVarPipelineDebugForceEvictImmediately(
+static TAutoConsoleVariable<int32> CVarPipelineDebugForceEvictImmediately(
 	TEXT("r.Vulkan.PipelineDebugForceEvictImmediately"),
 	0,
 	TEXT("1: Force all created PSOs to be evicted immediately. Only for debugging"),
@@ -63,13 +63,13 @@ TAutoConsoleVariable<int32> CVarPipelineDebugForceEvictImmediately(
 
 
 
-TAutoConsoleVariable<int32> CVarPipelineLRUCacheEvictBinaryPreloadScreen(
+static TAutoConsoleVariable<int32> CVarPipelineLRUCacheEvictBinaryPreloadScreen(
 	TEXT("r.Vulkan.PipelineLRUCacheEvictBinaryPreloadScreen"),
 	0,
 	TEXT("1: Use a preload screen while loading preevicted PSOs ala r.Vulkan.PipelineLRUCacheEvictBinary"),
 	ECVF_RenderThreadSafe);
 
-TAutoConsoleVariable<int32> CVarEnableLRU(
+static TAutoConsoleVariable<int32> CVarEnableLRU(
 	TEXT("r.Vulkan.EnablePipelineLRUCache"),
 	0,
 	TEXT("Pipeline LRU cache.\n")
@@ -77,7 +77,7 @@ TAutoConsoleVariable<int32> CVarEnableLRU(
 	TEXT("1: Enable LRU"),
 	ECVF_RenderThreadSafe | ECVF_ReadOnly);
 
-TAutoConsoleVariable<int32> CVarPipelineLRUCacheEvictBinary(
+static TAutoConsoleVariable<int32> CVarPipelineLRUCacheEvictBinary(
 	TEXT("r.Vulkan.PipelineLRUCacheEvictBinary"),
 	0,
 	TEXT("0: create pipelines in from the binary PSO cache and binary shader cache and evict them only as it fills up.\n")
@@ -85,13 +85,13 @@ TAutoConsoleVariable<int32> CVarPipelineLRUCacheEvictBinary(
 	ECVF_RenderThreadSafe | ECVF_ReadOnly);
 
 
-TAutoConsoleVariable<int32> CVarLRUMaxPipelineSize(
+static TAutoConsoleVariable<int32> CVarLRUMaxPipelineSize(
 	TEXT("r.Vulkan.PipelineLRUSize"),
 	LRU_MAX_PIPELINE_SIZE * 1024 * 1024,
 	TEXT("Maximum size of shader memory ."),
 	ECVF_RenderThreadSafe);
 
-TAutoConsoleVariable<int32> CVarLRUPipelineCapacity(
+static TAutoConsoleVariable<int32> CVarLRUPipelineCapacity(
 	TEXT("r.Vulkan.PipelineLRUCapactiy"),
 	LRU_PIPELINE_CAPACITY,
 	TEXT("Maximum no. of PSOs in LRU."),
@@ -104,15 +104,15 @@ TAutoConsoleVariable<int32> CVarLRUPipelineCapacity(
 static TAutoConsoleVariable<int32> GEnablePipelineCacheLoadCvar(
 	TEXT("r.Vulkan.PipelineCacheLoad"),
 	1,
-	TEXT("0 to disable loading the pipeline cache")
-	TEXT("1 to enable using pipeline cache")
+	TEXT("0 to disable loading the pipeline cache.\n")
+	TEXT("1 to enable using pipeline cache.\n")
 );
 
 static TAutoConsoleVariable<int32> GPipelineCacheFromShaderPipelineCacheCvar(
 	TEXT("r.Vulkan.PipelineCacheFromShaderPipelineCache"),
 	PLATFORM_ANDROID,
-	TEXT("0 look for a pipeline cache in the normal locations with the normal names.")
-	TEXT("1 tie the vulkan pipeline cache to the shader pipeline cache, use the PSOFC guid as part of the filename, etc."),
+	TEXT("0 look for a pipeline cache in the normal locations with the normal names.\n")
+	TEXT("1 tie the vulkan pipeline cache to the shader pipeline cache, use the PSOFC guid as part of the filename, etc.\n"),
 	ECVF_ReadOnly
 );
 
