VulkanRHI:  Add experimental support for graphics_pipeline_library on desktop SM6 (WIP):
- Requires VK_EXT_extended_dynamic_state, VK_EXT_extended_dynamic_state2, VK_EXT_extended_dynamic_state3 and VK_EXT_vertex_input_dynamic_state to simplify implementation.
- Requires VK_KHR_dynamic_rendering to remove the need for render passes.
- Requires bindless support (VK_EXT_descriptor_buffer) to simplify pipeline layouts.
- All these requirements allow us to build pipeline libraries for vs/ms/ps directly when shader is loaded.

#rnx

diff --git a/Engine/Source/Runtime/VulkanRHI/Private/VulkanPipeline.cpp b/Engine/Source/Runtime/VulkanRHI/Private/VulkanPipeline.cpp
--- a/Engine/Source/Runtime/VulkanRHI/Private/VulkanPipeline.cpp
+++ b/Engine/Source/Runtime/VulkanRHI/Private/VulkanPipeline.cpp
@@ -20,6 +20,7 @@
 #include "VulkanChunkedPipelineCache.h"
 #include "VulkanRenderTargetLayout.h"
 #include "VulkanRenderpass.h"
+#include "VulkanBindlessDescriptorManager.h"
 
 #define LRU_DEBUG 0
 #if !UE_BUILD_SHIPPING
@@ -165,6 +166,18 @@ static FAutoConsoleVariableRef GVulkanReleaseShaderModuleWhenEvictingPSOCVar(
 	,ECVF_RenderThreadSafe
 );
 
+
+int32 GVulkanGraphicPipelineLibraryLinkingMode = 0;
+static TAutoConsoleVariable<int32> GVulkanGraphicPipelineLibraryLinkingModeCvar(
+	TEXT("r.Vulkan.GraphicPipelineLibraryLinkingMode"),
+	0,
+	TEXT("0 to disable link time optimization (faster on CPU, slower on GPU).\n")
+	TEXT("1 to enable link time optimization (slower on CPU, faster on GPU).\n")
+);  // :todo-jn: Add "2": create unoptimized pipeline on first draw, but replace it eventually with optimized pipeline.
+
+
+
+
 template <typename TRHIType, typename TVulkanType>
 static inline FSHAHash GetShaderHash(TRHIType* RHIShader)
 {
@@ -293,6 +306,12 @@ FVulkanPipelineStateCacheManager::FVulkanPipelineStateCacheManager(FVulkanDevice
 {
 	bUseLRU = (int32)CVarEnableLRU.GetValueOnAnyThread() != 0;
 	LRUUsedPipelineMax = CVarLRUPipelineCapacity.GetValueOnAnyThread();
+
+	if (InDevice->SupportsGraphicPipelineLibraries())
+	{
+		CreateVertexInputLibraries();
+		CreateEmptyFragmentStageLibrary();
+	}
 }
 
 
@@ -309,6 +328,28 @@ FVulkanPipelineStateCacheManager::~FVulkanPipelineStateCacheManager()
 	}
 	DestroyCache();
 
+
+	if (VertexInputPipelineLibraryVS)
+	{
+		VulkanRHI::vkDestroyPipeline(Device->GetHandle(), VertexInputPipelineLibraryVS, VULKAN_CPU_ALLOCATOR);
+		VertexInputPipelineLibraryVS = VK_NULL_HANDLE;
+	}
+	if (VertexInputPipelineLibraryMS)
+	{
+		VulkanRHI::vkDestroyPipeline(Device->GetHandle(), VertexInputPipelineLibraryMS, VULKAN_CPU_ALLOCATOR);
+		VertexInputPipelineLibraryMS = VK_NULL_HANDLE;
+	}
+	if (EmptyFragmentStageLibrary)
+	{
+		VulkanRHI::vkDestroyPipeline(Device->GetHandle(), EmptyFragmentStageLibrary, VULKAN_CPU_ALLOCATOR);
+		EmptyFragmentStageLibrary = VK_NULL_HANDLE;
+	}
+	for (auto& Pair : FragmentOutputStatePipelineLibraries)
+	{
+		VulkanRHI::vkDestroyPipeline(Device->GetHandle(), Pair.Value, VULKAN_CPU_ALLOCATOR);
+	}
+	FragmentOutputStatePipelineLibraries.Empty();
+
 	// Only destroy layouts when quitting
 	for (auto& Pair : LayoutMap)
 	{
@@ -1314,8 +1355,293 @@ static TArray<VkDynamicState> FillDynamicStates(const FVulkanDevice& Device, con
 	return DynamicStates;
 }
 
+
+void FVulkanPipelineStateCacheManager::CreateVertexInputLibraries()
+{
+	INC_DWORD_STAT(STAT_VulkanNumVertexInputLibs);
+#if VULKAN_ENABLE_AGGRESSIVE_STATS
+	TRACE_CPUPROFILER_EVENT_SCOPE(VulkanVertexInputLibTime);
+#endif
+
+	VkGraphicsPipelineLibraryCreateInfoEXT LibraryInfo;
+	ZeroVulkanStruct(LibraryInfo, VK_STRUCTURE_TYPE_GRAPHICS_PIPELINE_LIBRARY_CREATE_INFO_EXT);
+	LibraryInfo.flags = VK_GRAPHICS_PIPELINE_LIBRARY_VERTEX_INPUT_INTERFACE_BIT_EXT;
+
+	// Set default values for unused states
+	VkPipelineVertexInputStateCreateInfo VertexInputStateCreateInfo;
+	ZeroVulkanStruct(VertexInputStateCreateInfo, VK_STRUCTURE_TYPE_PIPELINE_VERTEX_INPUT_STATE_CREATE_INFO);
+	VkPipelineInputAssemblyStateCreateInfo InputAssemblyStateCreateInfo;
+	ZeroVulkanStruct(InputAssemblyStateCreateInfo, VK_STRUCTURE_TYPE_PIPELINE_INPUT_ASSEMBLY_STATE_CREATE_INFO);
+
+	// See table for vertex input: https://docs.vulkan.org/guide/latest/dynamic_state_map.html
+	const VkDynamicState VertexInputDynamicStates[] = {
+		VK_DYNAMIC_STATE_PRIMITIVE_TOPOLOGY,
+		VK_DYNAMIC_STATE_VERTEX_INPUT_EXT
+	};
+
+	VkPipelineDynamicStateCreateInfo DynamicStateCreateInfo;
+	ZeroVulkanStruct(DynamicStateCreateInfo, VK_STRUCTURE_TYPE_PIPELINE_DYNAMIC_STATE_CREATE_INFO);
+	DynamicStateCreateInfo.dynamicStateCount = UE_ARRAY_COUNT(VertexInputDynamicStates);
+	DynamicStateCreateInfo.pDynamicStates = VertexInputDynamicStates;
+
+	VkGraphicsPipelineCreateInfo GraphicsPipelineCreateInfo;
+	ZeroVulkanStruct(GraphicsPipelineCreateInfo, VK_STRUCTURE_TYPE_GRAPHICS_PIPELINE_CREATE_INFO);
+	GraphicsPipelineCreateInfo.sType = VK_STRUCTURE_TYPE_GRAPHICS_PIPELINE_CREATE_INFO;
+	GraphicsPipelineCreateInfo.pNext = &LibraryInfo;
+	GraphicsPipelineCreateInfo.flags = VK_PIPELINE_CREATE_LIBRARY_BIT_KHR | VK_PIPELINE_CREATE_DESCRIPTOR_BUFFER_BIT_EXT | 
+		(GVulkanGraphicPipelineLibraryLinkingMode == 1 ? VK_PIPELINE_CREATE_RETAIN_LINK_TIME_OPTIMIZATION_INFO_BIT_EXT : 0);
+	GraphicsPipelineCreateInfo.layout = Device->GetBindlessDescriptorManager()->GetPipelineLayout();
+	GraphicsPipelineCreateInfo.pVertexInputState = &VertexInputStateCreateInfo;
+	GraphicsPipelineCreateInfo.pInputAssemblyState = &InputAssemblyStateCreateInfo;
+	GraphicsPipelineCreateInfo.pDynamicState = &DynamicStateCreateInfo;
+
+	VERIFYVULKANRESULT(VulkanRHI::vkCreateGraphicsPipelines(Device->GetHandle(), VK_NULL_HANDLE, 1, &GraphicsPipelineCreateInfo, VULKAN_CPU_ALLOCATOR, &VertexInputPipelineLibraryVS));
+
+	if (Device->GetOptionalExtensions().HasEXTMeshShader)
+	{
+		// Mesh shader is the same, but without the dynamic states (no topology, no vertex input)
+		DynamicStateCreateInfo.dynamicStateCount = 0;
+		DynamicStateCreateInfo.pDynamicStates = nullptr;
+
+		VERIFYVULKANRESULT(VulkanRHI::vkCreateGraphicsPipelines(Device->GetHandle(), VK_NULL_HANDLE, 1, &GraphicsPipelineCreateInfo, VULKAN_CPU_ALLOCATOR, &VertexInputPipelineLibraryMS));
+	}
+}
+
+void FVulkanPipelineStateCacheManager::CreateEmptyFragmentStageLibrary()
+{
+	INC_DWORD_STAT(STAT_VulkanNumPSLibs);
+#if VULKAN_ENABLE_AGGRESSIVE_STATS
+	TRACE_CPUPROFILER_EVENT_SCOPE(VulkanPSLibTime);
+#endif
+
+	VkGraphicsPipelineLibraryCreateInfoEXT LibraryInfo;
+	ZeroVulkanStruct(LibraryInfo, VK_STRUCTURE_TYPE_GRAPHICS_PIPELINE_LIBRARY_CREATE_INFO_EXT);
+	LibraryInfo.flags = VK_GRAPHICS_PIPELINE_LIBRARY_FRAGMENT_SHADER_BIT_EXT;
+
+	VkPipelineRenderingCreateInfo RenderingCreateInfo;
+	ZeroVulkanStruct(RenderingCreateInfo, VK_STRUCTURE_TYPE_PIPELINE_RENDERING_CREATE_INFO);
+	RenderingCreateInfo.viewMask = 0;  // :todo-jn: fill viewMask for multiview
+	LibraryInfo.pNext = &RenderingCreateInfo;
+
+	// See table for fragment shader: https://docs.vulkan.org/guide/latest/dynamic_state_map.html
+	const VkDynamicState FragmentShaderDynamicStates[] = {
+		VK_DYNAMIC_STATE_STENCIL_REFERENCE,
+		VK_DYNAMIC_STATE_DEPTH_BOUNDS,
+		VK_DYNAMIC_STATE_STENCIL_TEST_ENABLE, // ext1
+		VK_DYNAMIC_STATE_STENCIL_OP, // ext1
+		VK_DYNAMIC_STATE_DEPTH_COMPARE_OP, // ext1
+		VK_DYNAMIC_STATE_DEPTH_TEST_ENABLE, // ext1
+		VK_DYNAMIC_STATE_DEPTH_BOUNDS_TEST_ENABLE, // ext1
+		VK_DYNAMIC_STATE_DEPTH_WRITE_ENABLE, // ext1
+		VK_DYNAMIC_STATE_RASTERIZATION_SAMPLES_EXT, // ext3
+		VK_DYNAMIC_STATE_ALPHA_TO_COVERAGE_ENABLE_EXT // ext3
+	};
+
+	VkPipelineDynamicStateCreateInfo DynamicStateCreateInfo;
+	ZeroVulkanStruct(DynamicStateCreateInfo, VK_STRUCTURE_TYPE_PIPELINE_DYNAMIC_STATE_CREATE_INFO);
+	DynamicStateCreateInfo.dynamicStateCount = UE_ARRAY_COUNT(FragmentShaderDynamicStates);
+	DynamicStateCreateInfo.pDynamicStates = FragmentShaderDynamicStates;
+
+	// These struct's important fields are dynamic, set the rest to the default
+	VkPipelineMultisampleStateCreateInfo MultisampleStateCreateInfo;
+	ZeroVulkanStruct(MultisampleStateCreateInfo, VK_STRUCTURE_TYPE_PIPELINE_MULTISAMPLE_STATE_CREATE_INFO);
+	VkPipelineDepthStencilStateCreateInfo DepthStencilStateCreateInfo;
+	ZeroVulkanStruct(DepthStencilStateCreateInfo, VK_STRUCTURE_TYPE_PIPELINE_DEPTH_STENCIL_STATE_CREATE_INFO);
+
+	VkGraphicsPipelineCreateInfo GraphicsPipelineCreateInfo;
+	ZeroVulkanStruct(GraphicsPipelineCreateInfo, VK_STRUCTURE_TYPE_GRAPHICS_PIPELINE_CREATE_INFO);
+	GraphicsPipelineCreateInfo.sType = VK_STRUCTURE_TYPE_GRAPHICS_PIPELINE_CREATE_INFO;
+	GraphicsPipelineCreateInfo.pNext = &LibraryInfo;
+	GraphicsPipelineCreateInfo.flags = VK_PIPELINE_CREATE_LIBRARY_BIT_KHR | VK_PIPELINE_CREATE_DESCRIPTOR_BUFFER_BIT_EXT | 
+		(GVulkanGraphicPipelineLibraryLinkingMode == 1 ? VK_PIPELINE_CREATE_RETAIN_LINK_TIME_OPTIMIZATION_INFO_BIT_EXT : 0);
+	GraphicsPipelineCreateInfo.pMultisampleState = &MultisampleStateCreateInfo;
+	GraphicsPipelineCreateInfo.pDepthStencilState = &DepthStencilStateCreateInfo;
+	GraphicsPipelineCreateInfo.layout = Device->GetBindlessDescriptorManager()->GetPipelineLayout();
+	GraphicsPipelineCreateInfo.pDynamicState = &DynamicStateCreateInfo;
+
+	// Add the NullPS for platforms that don't allow a null pixel stage
+	VkShaderModuleCreateInfo ModuleCreateInfo;
+	VkPipelineShaderStageCreateInfo ShaderStageCreateInfo;
+	TArrayView<uint32> SpirvCode;
+	ANSICHAR EntryPoint[24];
+	if (!FVulkanPlatform::SupportsNullPixelShader())
+	{
+		ZeroVulkanStruct(ModuleCreateInfo, VK_STRUCTURE_TYPE_SHADER_MODULE_CREATE_INFO);
+		ZeroVulkanStruct(ShaderStageCreateInfo, VK_STRUCTURE_TYPE_PIPELINE_SHADER_STAGE_CREATE_INFO);
+
+		FVulkanPixelShader* PixelShader = ResourceCast(TShaderMapRef<FNULLPS>(GetGlobalShaderMap(GMaxRHIFeatureLevel)).GetPixelShader());
+
+		SpirvCode = PixelShader->GetSpirvCode().GetCodeView();
+		ModuleCreateInfo.codeSize = SpirvCode.Num() * sizeof(uint32);
+		ModuleCreateInfo.pCode = SpirvCode.GetData();
+
+		ShaderStageCreateInfo.pNext = &ModuleCreateInfo;
+		ShaderStageCreateInfo.stage = VK_SHADER_STAGE_FRAGMENT_BIT;
+
+		PixelShader->GetEntryPoint(EntryPoint, 24);
+		ShaderStageCreateInfo.pName = EntryPoint;
+
+		GraphicsPipelineCreateInfo.pStages = &ShaderStageCreateInfo;
+	}
+
+	VERIFYVULKANRESULT(VulkanRHI::vkCreateGraphicsPipelines(Device->GetHandle(), VK_NULL_HANDLE, 1, &GraphicsPipelineCreateInfo, VULKAN_CPU_ALLOCATOR, &EmptyFragmentStageLibrary));
+}
+
+void FillDynamicRenderingCreateInfo(
+	const FGfxPipelineDesc::FRenderTargets& InRenderTargets, 
+	VkPipelineRenderingCreateInfo& OutPipelineRenderingCreateInfo,
+	TStaticArray<VkFormat, MaxSimultaneousRenderTargets>& OutColorAttachmentFormats)
+{
+	if (InRenderTargets.NumColorAttachments > 0)
+	{
+		for (uint8 ColorIndex = 0; ColorIndex < InRenderTargets.NumColorAttachments; ++ColorIndex)
+		{
+			const uint32 Attachment = InRenderTargets.ColorAttachments[ColorIndex].Attachment;
+			OutColorAttachmentFormats[ColorIndex] = (VkFormat)InRenderTargets.Descriptions[Attachment].Format;
+		}
+
+		OutPipelineRenderingCreateInfo.pColorAttachmentFormats = OutColorAttachmentFormats.GetData();
+		OutPipelineRenderingCreateInfo.colorAttachmentCount = InRenderTargets.NumColorAttachments;
+	}
+
+	if (InRenderTargets.bHasDepthStencil)
+	{
+		// Depth
+		const uint32 Attachment = InRenderTargets.Depth.Attachment;
+		OutPipelineRenderingCreateInfo.depthAttachmentFormat = (VkFormat)InRenderTargets.Descriptions[Attachment].Format;
+
+		// Stencil
+		if (VulkanRHI::VulkanFormatHasStencil(OutPipelineRenderingCreateInfo.depthAttachmentFormat))
+		{
+			// :todo-jn: leave unused attachements to VK_FORMAT_UNDEFINED
+			OutPipelineRenderingCreateInfo.stencilAttachmentFormat = OutPipelineRenderingCreateInfo.depthAttachmentFormat;
+		}
+	}
+}
+
+VkPipeline FVulkanPipelineStateCacheManager::FindOrCreateFragmentOutputStateLibrary(const FVulkanGraphicsPipelineState& GraphicsPipelineState)
+{
+	TStaticArray<VkFormat, MaxSimultaneousRenderTargets> ColorAttachmentFormats;
+	VkPipelineRenderingCreateInfo PipelineRenderingCreateInfo;
+	ZeroVulkanStruct(PipelineRenderingCreateInfo, VK_STRUCTURE_TYPE_PIPELINE_RENDERING_CREATE_INFO);
+	FillDynamicRenderingCreateInfo(GraphicsPipelineState.Desc.RenderTargets, PipelineRenderingCreateInfo, ColorAttachmentFormats);
+
+	uint32 LibraryCRC = FCrc::MemCrc32(PipelineRenderingCreateInfo.pColorAttachmentFormats, PipelineRenderingCreateInfo.colorAttachmentCount * sizeof(VkFormat));
+	LibraryCRC = FCrc::MemCrc32(&PipelineRenderingCreateInfo.depthAttachmentFormat, sizeof(VkFormat), LibraryCRC);
+	LibraryCRC = FCrc::MemCrc32(&PipelineRenderingCreateInfo.stencilAttachmentFormat, sizeof(VkFormat), LibraryCRC);
+
+	FScopeLock Lock(&FragmentOutputStateCS);
+	if (VkPipeline* ExistingPipelineLibrary = FragmentOutputStatePipelineLibraries.Find(LibraryCRC))
+	{
+		return *ExistingPipelineLibrary;
+	}
+
+	INC_DWORD_STAT(STAT_VulkanNumFragOutputStateLibs);
+#if VULKAN_ENABLE_AGGRESSIVE_STATS
+	TRACE_CPUPROFILER_EVENT_SCOPE(VulkanFragOutputStateLibTime);
+#endif
+
+	VkGraphicsPipelineLibraryCreateInfoEXT LibraryInfo;
+	ZeroVulkanStruct(LibraryInfo, VK_STRUCTURE_TYPE_GRAPHICS_PIPELINE_LIBRARY_CREATE_INFO_EXT);
+	LibraryInfo.flags = VK_GRAPHICS_PIPELINE_LIBRARY_FRAGMENT_OUTPUT_INTERFACE_BIT_EXT;
+	LibraryInfo.pNext = &PipelineRenderingCreateInfo;
+
+	// Set default values for unused states
+	VkPipelineColorBlendStateCreateInfo ColorBlendStateCreateInfo;
+	ZeroVulkanStruct(ColorBlendStateCreateInfo, VK_STRUCTURE_TYPE_PIPELINE_COLOR_BLEND_STATE_CREATE_INFO);
+	ColorBlendStateCreateInfo.blendConstants[0] = 1.0f;
+	ColorBlendStateCreateInfo.blendConstants[1] = 1.0f;
+	ColorBlendStateCreateInfo.blendConstants[2] = 1.0f;
+	ColorBlendStateCreateInfo.blendConstants[3] = 1.0f;
+
+	VkPipelineMultisampleStateCreateInfo MultisampleStateCreateInfo;
+	ZeroVulkanStruct(MultisampleStateCreateInfo, VK_STRUCTURE_TYPE_PIPELINE_MULTISAMPLE_STATE_CREATE_INFO);
+
+	// See table for fragment output state: https://docs.vulkan.org/guide/latest/dynamic_state_map.html
+	const VkDynamicState FragmentOutputStateDynamicStates[] = {
+		VK_DYNAMIC_STATE_RASTERIZATION_SAMPLES_EXT,
+		VK_DYNAMIC_STATE_ALPHA_TO_COVERAGE_ENABLE_EXT,
+		VK_DYNAMIC_STATE_COLOR_BLEND_ENABLE_EXT,
+		VK_DYNAMIC_STATE_COLOR_WRITE_MASK_EXT,
+		VK_DYNAMIC_STATE_COLOR_BLEND_EQUATION_EXT
+	};
+
+	VkPipelineDynamicStateCreateInfo DynamicStateCreateInfo;
+	ZeroVulkanStruct(DynamicStateCreateInfo, VK_STRUCTURE_TYPE_PIPELINE_DYNAMIC_STATE_CREATE_INFO);
+	DynamicStateCreateInfo.dynamicStateCount = UE_ARRAY_COUNT(FragmentOutputStateDynamicStates);
+	DynamicStateCreateInfo.pDynamicStates = FragmentOutputStateDynamicStates;
+
+	VkGraphicsPipelineCreateInfo GraphicsPipelineCreateInfo;
+	ZeroVulkanStruct(GraphicsPipelineCreateInfo, VK_STRUCTURE_TYPE_GRAPHICS_PIPELINE_CREATE_INFO);
+	GraphicsPipelineCreateInfo.sType = VK_STRUCTURE_TYPE_GRAPHICS_PIPELINE_CREATE_INFO;
+	GraphicsPipelineCreateInfo.pNext = &LibraryInfo;
+	GraphicsPipelineCreateInfo.flags = VK_PIPELINE_CREATE_LIBRARY_BIT_KHR | VK_PIPELINE_CREATE_DESCRIPTOR_BUFFER_BIT_EXT |
+		(GVulkanGraphicPipelineLibraryLinkingMode == 1 ? VK_PIPELINE_CREATE_RETAIN_LINK_TIME_OPTIMIZATION_INFO_BIT_EXT : 0);
+	GraphicsPipelineCreateInfo.layout = Device->GetBindlessDescriptorManager()->GetPipelineLayout();
+	GraphicsPipelineCreateInfo.pColorBlendState = &ColorBlendStateCreateInfo;
+	GraphicsPipelineCreateInfo.pMultisampleState = &MultisampleStateCreateInfo;
+	GraphicsPipelineCreateInfo.pDynamicState = &DynamicStateCreateInfo;
+
+	VkPipeline FragmentOutputStatePipelineLibrary = VK_NULL_HANDLE;
+	VERIFYVULKANRESULT(VulkanRHI::vkCreateGraphicsPipelines(Device->GetHandle(), VK_NULL_HANDLE, 1, &GraphicsPipelineCreateInfo, VULKAN_CPU_ALLOCATOR, &FragmentOutputStatePipelineLibrary));
+	FragmentOutputStatePipelineLibraries.Add(LibraryCRC, FragmentOutputStatePipelineLibrary);
+	return FragmentOutputStatePipelineLibrary;
+}
+
+
+bool FVulkanPipelineStateCacheManager::CreateGfxPipelineFromLibrary(FVulkanGraphicsPipelineState& GraphicsPipelineState, FVulkanShader* Shaders[ShaderStage::NumGraphicsStages], FGraphicsPipelineStateInitializer::EPSOPrecacheCompileType PSOCompileType)
+{
+	// Only tackle the most common case for now, use regular pipelines for the rest
+	if (Shaders[ShaderStage::Task] || Shaders[ShaderStage::Geometry])
+	{
+		return false;
+	}
+
+#if VULKAN_ENABLE_AGGRESSIVE_STATS
+	TRACE_CPUPROFILER_EVENT_SCOPE(VulkanLibPSOCompileTime);
+#endif
+
+	FVulkanShader* VertexShader = Shaders[ShaderStage::Vertex];
+	FVulkanShader* PixelShader = Shaders[ShaderStage::Pixel];
+	FVulkanShader* MeshShader = Shaders[ShaderStage::Mesh];
+
+	TStaticArray<VkPipeline, 4> Libraries;
+	Libraries[0] = MeshShader ? VertexInputPipelineLibraryMS : VertexInputPipelineLibraryVS;
+	Libraries[1] = MeshShader ? MeshShader->PipelineLibrary : VertexShader->PipelineLibrary;
+	Libraries[2] = PixelShader ? PixelShader->PipelineLibrary : EmptyFragmentStageLibrary;
+	Libraries[3] = FindOrCreateFragmentOutputStateLibrary(GraphicsPipelineState);
+
+	VkPipelineLibraryCreateInfoKHR LibraryCreateInfo;
+	ZeroVulkanStruct(LibraryCreateInfo, VK_STRUCTURE_TYPE_PIPELINE_LIBRARY_CREATE_INFO_KHR);
+	LibraryCreateInfo.libraryCount = Libraries.Num();
+	LibraryCreateInfo.pLibraries = Libraries.GetData();
+
+	VkGraphicsPipelineCreateInfo ExecutablePipelineCreateInfo;
+	ZeroVulkanStruct(ExecutablePipelineCreateInfo, VK_STRUCTURE_TYPE_GRAPHICS_PIPELINE_CREATE_INFO);
+	ExecutablePipelineCreateInfo.pNext = &LibraryCreateInfo;
+	ExecutablePipelineCreateInfo.flags = VK_PIPELINE_CREATE_DESCRIPTOR_BUFFER_BIT_EXT | 
+		(GVulkanGraphicPipelineLibraryLinkingMode == 1 ? VK_PIPELINE_CREATE_LINK_TIME_OPTIMIZATION_BIT_EXT : 0);
+	ExecutablePipelineCreateInfo.layout = Device->GetBindlessDescriptorManager()->GetPipelineLayout();
+
+	const VkResult Result = VulkanRHI::vkCreateGraphicsPipelines(Device->GetHandle(), nullptr, 1, &ExecutablePipelineCreateInfo, VULKAN_CPU_ALLOCATOR, &GraphicsPipelineState.VulkanPipeline);
+	VERIFYVULKANRESULT(Result);
+	return (Result == VK_SUCCESS);
+}
+
 bool FVulkanPipelineStateCacheManager::CreateGfxPipelineFromEntry(FVulkanGraphicsPipelineState* PSO, FVulkanShader* Shaders[ShaderStage::NumGraphicsStages], FGraphicsPipelineStateInitializer::EPSOPrecacheCompileType PSOCompileType)
 {
+	// See if it can be created from a pipeline library first
+	if (Device->SupportsGraphicPipelineLibraries() && CreateGfxPipelineFromLibrary(*PSO, Shaders, PSOCompileType))
+	{
+		INC_DWORD_STAT(STAT_VulkanLibPSOCompileCount);
+		return true;
+	}
+
+	INC_DWORD_STAT(STAT_VulkanRegPSOCompileCount);
+#if VULKAN_ENABLE_AGGRESSIVE_STATS
+	TRACE_CPUPROFILER_EVENT_SCOPE(VulkanRegPSOCompileTime);
+#endif
+
 	VkPipeline* Pipeline = &PSO->VulkanPipeline;
 	const FGfxPipelineDesc* GfxEntry = &PSO->Desc;
 	if (Shaders[ShaderStage::Pixel] == nullptr && !FVulkanPlatform::SupportsNullPixelShader())
@@ -1445,32 +1771,9 @@ bool FVulkanPipelineStateCacheManager::CreateGfxPipelineFromEntry(FVulkanGraphic
 	if (Device->GetOptionalExtensions().HasKHRDynamicRendering)
 	{
 		ZeroVulkanStruct(PipelineRenderingCreateInfo, VK_STRUCTURE_TYPE_PIPELINE_RENDERING_CREATE_INFO);
+		FillDynamicRenderingCreateInfo(PSO->Desc.RenderTargets, PipelineRenderingCreateInfo, ColorAttachmentFormats);
 		PipelineRenderingCreateInfo.pNext = PipelineInfo.pNext;
 		PipelineInfo.pNext = &PipelineRenderingCreateInfo;
-
-		PipelineRenderingCreateInfo.pColorAttachmentFormats = ColorAttachmentFormats.GetData();
-		PipelineRenderingCreateInfo.colorAttachmentCount = PSO->Desc.RenderTargets.NumColorAttachments;
-		for (uint8 ColorIndex = 0; ColorIndex < PSO->Desc.RenderTargets.NumColorAttachments; ++ColorIndex)
-		{
-			const uint32 Attachment = PSO->Desc.RenderTargets.ColorAttachments[ColorIndex].Attachment;
-			const FGfxPipelineDesc::FRenderTargets::FAttachmentDesc& Descripton = PSO->Desc.RenderTargets.Descriptions[Attachment];
-			ColorAttachmentFormats[ColorIndex] = (VkFormat)Descripton.Format;
-		}
-
-		if (PSO->Desc.RenderTargets.bHasDepthStencil)
-		{
-			// Depth
-			const uint32 Attachment = PSO->Desc.RenderTargets.Depth.Attachment;
-			const FGfxPipelineDesc::FRenderTargets::FAttachmentDesc& Descripton = PSO->Desc.RenderTargets.Descriptions[Attachment];
-			PipelineRenderingCreateInfo.depthAttachmentFormat = (VkFormat)Descripton.Format;
-
-			// Stencil
-			if (VulkanRHI::VulkanFormatHasStencil(PipelineRenderingCreateInfo.depthAttachmentFormat))
-			{
-				// :todo-jn: leave unused attachements to VK_FORMAT_UNDEFINED
-				PipelineRenderingCreateInfo.stencilAttachmentFormat = PipelineRenderingCreateInfo.depthAttachmentFormat;
-			}
-		}
 	}
 	else
 	{
diff --git a/Engine/Source/Runtime/VulkanRHI/Private/VulkanPipeline.h b/Engine/Source/Runtime/VulkanRHI/Private/VulkanPipeline.h
--- a/Engine/Source/Runtime/VulkanRHI/Private/VulkanPipeline.h
+++ b/Engine/Source/Runtime/VulkanRHI/Private/VulkanPipeline.h
@@ -522,6 +522,18 @@ private:
 	FVulkanComputePipeline* RHICreateComputePipelineState(FRHIComputeShader* ComputeShaderRHI);
 	void NotifyDeletedGraphicsPSO(FRHIGraphicsPipelineState* PSO);
 	bool CreateGfxPipelineFromEntry(FVulkanGraphicsPipelineState* PSO, FVulkanShader* Shaders[ShaderStage::NumGraphicsStages], FGraphicsPipelineStateInitializer::EPSOPrecacheCompileType PSOCompileType);
+	bool CreateGfxPipelineFromLibrary(FVulkanGraphicsPipelineState& GraphicsPipelineState, FVulkanShader* Shaders[ShaderStage::NumGraphicsStages], FGraphicsPipelineStateInitializer::EPSOPrecacheCompileType PSOCompileType);
+
+	void CreateVertexInputLibraries();
+	VkPipeline VertexInputPipelineLibraryVS = VK_NULL_HANDLE;
+	VkPipeline VertexInputPipelineLibraryMS = VK_NULL_HANDLE;
+
+	void CreateEmptyFragmentStageLibrary();
+	VkPipeline EmptyFragmentStageLibrary = VK_NULL_HANDLE;
+
+	VkPipeline FindOrCreateFragmentOutputStateLibrary(const FVulkanGraphicsPipelineState& GraphicsPipelineState);
+	TMap<uint32, VkPipeline> FragmentOutputStatePipelineLibraries;
+	FCriticalSection FragmentOutputStateCS;
 
 	VkResult CreateVKPipeline(FVulkanGraphicsPipelineState* PSO, FVulkanShader* Shaders[ShaderStage::NumGraphicsStages], const VkGraphicsPipelineCreateInfo& PipelineInfo, FGraphicsPipelineStateInitializer::EPSOPrecacheCompileType PSOCompileType);
 
diff --git a/Engine/Source/Runtime/VulkanRHI/Private/VulkanRHIPrivate.h b/Engine/Source/Runtime/VulkanRHI/Private/VulkanRHIPrivate.h
--- a/Engine/Source/Runtime/VulkanRHI/Private/VulkanRHIPrivate.h
+++ b/Engine/Source/Runtime/VulkanRHI/Private/VulkanRHIPrivate.h
@@ -142,7 +142,7 @@ union UNvidiaDriverVersion
 void VulkanSetImageLayout(FVulkanCommandBuffer* CmdBuffer, VkImage Image, VkImageLayout OldLayout, VkImageLayout NewLayout, const VkImageSubresourceRange& SubresourceRange);
 
 
-
+// Regular pipeline creation stats
 DECLARE_STATS_GROUP(TEXT("Vulkan PSO"), STATGROUP_VulkanPSO, STATCAT_Advanced);
 DECLARE_DWORD_ACCUMULATOR_STAT_EXTERN(TEXT("PSO LRU Elements"), STAT_VulkanNumPSOLRU, STATGROUP_VulkanPSO, );
 DECLARE_DWORD_ACCUMULATOR_STAT_EXTERN(TEXT("PSO LRU Size"), STAT_VulkanNumPSOLRUSize, STATGROUP_VulkanPSO, );
@@ -157,7 +157,19 @@ DECLARE_CYCLE_STAT_EXTERN(TEXT("PSO Creation time"), STAT_VulkanPSOCreationTime,
 DECLARE_CYCLE_STAT_EXTERN(TEXT("PSO Vulkan Creation time"), STAT_VulkanPSOVulkanCreationTime, STATGROUP_VulkanPSO, );
 
 
-// Stats
+// Graphic pipeline library stats
+DECLARE_STATS_GROUP(TEXT("Vulkan Graphic Pipeline Libraries"), STATGROUP_VulkanGraphicPipelineLibs, STATCAT_Advanced);
+DECLARE_DWORD_ACCUMULATOR_STAT_EXTERN(TEXT("Regular PSO Compile Count"), STAT_VulkanRegPSOCompileCount, STATGROUP_VulkanGraphicPipelineLibs, );
+DECLARE_DWORD_ACCUMULATOR_STAT_EXTERN(TEXT("Library PSO Compile Count"), STAT_VulkanLibPSOCompileCount, STATGROUP_VulkanGraphicPipelineLibs, );
+DECLARE_DWORD_ACCUMULATOR_STAT_EXTERN(TEXT("Num VertexInput Libraries"), STAT_VulkanNumVertexInputLibs, STATGROUP_VulkanGraphicPipelineLibs, );
+DECLARE_DWORD_ACCUMULATOR_STAT_EXTERN(TEXT("Num VS Libraries"), STAT_VulkanNumVSLibs, STATGROUP_VulkanGraphicPipelineLibs, );
+DECLARE_DWORD_ACCUMULATOR_STAT_EXTERN(TEXT("Num MS Libraries"), STAT_VulkanNumMSLibs, STATGROUP_VulkanGraphicPipelineLibs, );
+DECLARE_DWORD_ACCUMULATOR_STAT_EXTERN(TEXT("Num PS Libraries"), STAT_VulkanNumPSLibs, STATGROUP_VulkanGraphicPipelineLibs, );
+DECLARE_DWORD_ACCUMULATOR_STAT_EXTERN(TEXT("Num FragmentOutputState Libraries"), STAT_VulkanNumFragOutputStateLibs, STATGROUP_VulkanGraphicPipelineLibs, );
+DECLARE_DWORD_ACCUMULATOR_STAT_EXTERN(TEXT("Num Shader Modules"), STAT_VulkanNumShaderModule, STATGROUP_VulkanGraphicPipelineLibs, );
+
+
+// General Stats
 DECLARE_STATS_GROUP(TEXT("Vulkan RHI"), STATGROUP_VulkanRHI, STATCAT_Advanced);
 DECLARE_CYCLE_STAT_EXTERN(TEXT("Draw call time"), STAT_VulkanDrawCallTime, STATGROUP_VulkanRHI, );
 DECLARE_CYCLE_STAT_EXTERN(TEXT("Dispatch call time"), STAT_VulkanDispatchCallTime, STATGROUP_VulkanRHI, );
diff --git a/Engine/Source/Runtime/VulkanRHI/Private/VulkanResources.h b/Engine/Source/Runtime/VulkanRHI/Private/VulkanResources.h
--- a/Engine/Source/Runtime/VulkanRHI/Private/VulkanResources.h
+++ b/Engine/Source/Runtime/VulkanRHI/Private/VulkanResources.h
@@ -214,6 +214,8 @@ protected:
 	const EShaderFrequency			Frequency;
 	bool							bUsesBindless = false;
 
+	VkPipeline						PipelineLibrary = VK_NULL_HANDLE;
+
 protected:
 	class FSpirvContainer
 	{
diff --git a/Engine/Source/Runtime/VulkanRHI/Private/VulkanShaders.cpp b/Engine/Source/Runtime/VulkanRHI/Private/VulkanShaders.cpp
--- a/Engine/Source/Runtime/VulkanRHI/Private/VulkanShaders.cpp
+++ b/Engine/Source/Runtime/VulkanRHI/Private/VulkanShaders.cpp
@@ -12,6 +12,11 @@
 #include "VulkanLLM.h"
 #include "VulkanDescriptorSets.h"
 #include "RHICoreShader.h"
+#include "VulkanBindlessDescriptorManager.h"
+
+
+extern int32 GVulkanGraphicPipelineLibraryLinkingMode;
+
 
 static int32 GVulkanCompressSPIRV = 0;
 static FAutoConsoleVariableRef GVulkanCompressSPIRVCVar(
@@ -51,6 +56,158 @@ static void ReadShaderOptionalData(FShaderCodeReader& ShaderCode, ShaderType* RH
 #endif
 }
 
+static VkPipeline CreatePreRasterShaderLibrary(FVulkanDevice* Device, FVulkanShader* Shader, VkShaderStageFlagBits StageFlag)
+{
+	if (StageFlag == VK_SHADER_STAGE_VERTEX_BIT)
+	{
+		INC_DWORD_STAT(STAT_VulkanNumVSLibs);
+	}
+	else if (StageFlag == VK_SHADER_STAGE_MESH_BIT_EXT)
+	{
+		INC_DWORD_STAT(STAT_VulkanNumMSLibs);
+	}
+	else
+	{
+		checkNoEntry();
+	}
+
+#if VULKAN_ENABLE_AGGRESSIVE_STATS
+	TRACE_CPUPROFILER_EVENT_SCOPE_CONDITIONAL(VulkanVSLibTime, (StageFlag == VK_SHADER_STAGE_VERTEX_BIT));
+	TRACE_CPUPROFILER_EVENT_SCOPE_CONDITIONAL(VulkanMSLibTime, (StageFlag == VK_SHADER_STAGE_MESH_BIT_EXT));
+#endif
+
+	VkShaderModuleCreateInfo ModuleCreateInfo;
+	ZeroVulkanStruct(ModuleCreateInfo, VK_STRUCTURE_TYPE_SHADER_MODULE_CREATE_INFO);
+
+	TArrayView<uint32> SpirvCode = Shader->GetSpirvCode().GetCodeView();
+	ModuleCreateInfo.codeSize = SpirvCode.Num() * sizeof(uint32);
+	ModuleCreateInfo.pCode = SpirvCode.GetData();
+
+	VkGraphicsPipelineLibraryCreateInfoEXT LibraryInfo;
+	ZeroVulkanStruct(LibraryInfo, VK_STRUCTURE_TYPE_GRAPHICS_PIPELINE_LIBRARY_CREATE_INFO_EXT);
+	LibraryInfo.flags = VK_GRAPHICS_PIPELINE_LIBRARY_PRE_RASTERIZATION_SHADERS_BIT_EXT;
+
+	VkPipelineShaderStageCreateInfo ShaderStageCreateInfo;
+	ZeroVulkanStruct(ShaderStageCreateInfo, VK_STRUCTURE_TYPE_PIPELINE_SHADER_STAGE_CREATE_INFO);
+	ShaderStageCreateInfo.pNext = &ModuleCreateInfo;
+	ShaderStageCreateInfo.stage = StageFlag;
+
+	ANSICHAR EntryPoint[24];
+	Shader->GetEntryPoint(EntryPoint, 24);
+	ShaderStageCreateInfo.pName = EntryPoint;
+
+	// See table for pre-raster: https://docs.vulkan.org/guide/latest/dynamic_state_map.html
+	const VkDynamicState PreRasterDynamicStates[] = {
+		VK_DYNAMIC_STATE_DEPTH_BIAS,
+		VK_DYNAMIC_STATE_CULL_MODE, // ext1
+		VK_DYNAMIC_STATE_VIEWPORT_WITH_COUNT_EXT, // ext1
+		VK_DYNAMIC_STATE_SCISSOR_WITH_COUNT_EXT, // ext1
+		VK_DYNAMIC_STATE_DEPTH_BIAS_ENABLE,  // ext2
+		VK_DYNAMIC_STATE_POLYGON_MODE_EXT  // ext3
+	};
+
+	VkPipelineDynamicStateCreateInfo DynamicStateCreateInfo;
+	ZeroVulkanStruct(DynamicStateCreateInfo, VK_STRUCTURE_TYPE_PIPELINE_DYNAMIC_STATE_CREATE_INFO);
+	DynamicStateCreateInfo.dynamicStateCount = UE_ARRAY_COUNT(PreRasterDynamicStates);
+	DynamicStateCreateInfo.pDynamicStates = PreRasterDynamicStates;
+
+	// These struct's important fields are dynamic, set the rest to the default
+	VkPipelineRasterizationStateCreateInfo RasterizerState;
+	FVulkanRasterizerState::ResetCreateInfo(RasterizerState);
+
+	VkGraphicsPipelineCreateInfo GraphicsPipelineCreateInfo;
+	ZeroVulkanStruct(GraphicsPipelineCreateInfo, VK_STRUCTURE_TYPE_GRAPHICS_PIPELINE_CREATE_INFO);
+	GraphicsPipelineCreateInfo.sType = VK_STRUCTURE_TYPE_GRAPHICS_PIPELINE_CREATE_INFO;
+	GraphicsPipelineCreateInfo.pNext = &LibraryInfo;
+	GraphicsPipelineCreateInfo.flags = VK_PIPELINE_CREATE_LIBRARY_BIT_KHR | VK_PIPELINE_CREATE_DESCRIPTOR_BUFFER_BIT_EXT | 
+		(GVulkanGraphicPipelineLibraryLinkingMode == 1 ? VK_PIPELINE_CREATE_RETAIN_LINK_TIME_OPTIMIZATION_INFO_BIT_EXT : 0);
+	GraphicsPipelineCreateInfo.stageCount = 1;
+	GraphicsPipelineCreateInfo.pStages = &ShaderStageCreateInfo;
+	GraphicsPipelineCreateInfo.layout = Device->GetBindlessDescriptorManager()->GetPipelineLayout();
+	GraphicsPipelineCreateInfo.pDynamicState = &DynamicStateCreateInfo;
+	GraphicsPipelineCreateInfo.pRasterizationState = &RasterizerState;
+
+	VkPipeline PipelineLibrary = VK_NULL_HANDLE;
+	VERIFYVULKANRESULT(VulkanRHI::vkCreateGraphicsPipelines(Device->GetHandle(), VK_NULL_HANDLE, 1, &GraphicsPipelineCreateInfo, VULKAN_CPU_ALLOCATOR, &PipelineLibrary));
+	return PipelineLibrary;
+}
+
+static VkPipeline CreatePixelShaderLibrary(FVulkanDevice* Device, FVulkanShader* PixelShader)
+{
+	INC_DWORD_STAT(STAT_VulkanNumPSLibs);
+
+#if VULKAN_ENABLE_AGGRESSIVE_STATS
+	TRACE_CPUPROFILER_EVENT_SCOPE(VulkanPSLibTime);
+#endif
+
+	VkShaderModuleCreateInfo ModuleCreateInfo;
+	ZeroVulkanStruct(ModuleCreateInfo, VK_STRUCTURE_TYPE_SHADER_MODULE_CREATE_INFO);
+
+	TArrayView<uint32> SpirvCode = PixelShader->GetSpirvCode().GetCodeView();
+	ModuleCreateInfo.codeSize = SpirvCode.Num() * sizeof(uint32);
+	ModuleCreateInfo.pCode = SpirvCode.GetData();
+
+	VkGraphicsPipelineLibraryCreateInfoEXT LibraryInfo;
+	ZeroVulkanStruct(LibraryInfo, VK_STRUCTURE_TYPE_GRAPHICS_PIPELINE_LIBRARY_CREATE_INFO_EXT);
+	LibraryInfo.flags = VK_GRAPHICS_PIPELINE_LIBRARY_FRAGMENT_SHADER_BIT_EXT;
+
+	VkPipelineRenderingCreateInfo RenderingCreateInfo;  
+	ZeroVulkanStruct(RenderingCreateInfo, VK_STRUCTURE_TYPE_PIPELINE_RENDERING_CREATE_INFO);
+	RenderingCreateInfo.viewMask = 0;  // :todo-jn: fill viewMask for multiview
+	LibraryInfo.pNext = &RenderingCreateInfo;
+
+	VkPipelineShaderStageCreateInfo ShaderStageCreateInfo;
+	ZeroVulkanStruct(ShaderStageCreateInfo, VK_STRUCTURE_TYPE_PIPELINE_SHADER_STAGE_CREATE_INFO);
+	ShaderStageCreateInfo.pNext = &ModuleCreateInfo;
+	ShaderStageCreateInfo.stage = VK_SHADER_STAGE_FRAGMENT_BIT;
+
+	ANSICHAR EntryPoint[24];
+	PixelShader->GetEntryPoint(EntryPoint, 24);
+	ShaderStageCreateInfo.pName = EntryPoint;
+
+	// See table for fragment shader: https://docs.vulkan.org/guide/latest/dynamic_state_map.html
+	const VkDynamicState FragmentShaderDynamicStates[] = {
+		VK_DYNAMIC_STATE_STENCIL_REFERENCE,
+		VK_DYNAMIC_STATE_DEPTH_BOUNDS,
+		VK_DYNAMIC_STATE_STENCIL_TEST_ENABLE, // ext1
+		VK_DYNAMIC_STATE_STENCIL_OP, // ext1
+		VK_DYNAMIC_STATE_DEPTH_COMPARE_OP, // ext1
+		VK_DYNAMIC_STATE_DEPTH_TEST_ENABLE, // ext1
+		VK_DYNAMIC_STATE_DEPTH_BOUNDS_TEST_ENABLE, // ext1
+		VK_DYNAMIC_STATE_DEPTH_WRITE_ENABLE, // ext1
+		VK_DYNAMIC_STATE_RASTERIZATION_SAMPLES_EXT, // ext3
+		VK_DYNAMIC_STATE_ALPHA_TO_COVERAGE_ENABLE_EXT // ext3
+	};
+
+	VkPipelineDynamicStateCreateInfo DynamicStateCreateInfo;
+	ZeroVulkanStruct(DynamicStateCreateInfo, VK_STRUCTURE_TYPE_PIPELINE_DYNAMIC_STATE_CREATE_INFO);
+	DynamicStateCreateInfo.dynamicStateCount = UE_ARRAY_COUNT(FragmentShaderDynamicStates);
+	DynamicStateCreateInfo.pDynamicStates = FragmentShaderDynamicStates;
+
+	// These struct's important fields are dynamic, set the rest to the default
+	VkPipelineMultisampleStateCreateInfo MultisampleStateCreateInfo;
+	ZeroVulkanStruct(MultisampleStateCreateInfo, VK_STRUCTURE_TYPE_PIPELINE_MULTISAMPLE_STATE_CREATE_INFO);
+	VkPipelineDepthStencilStateCreateInfo DepthStencilStateCreateInfo;
+	ZeroVulkanStruct(DepthStencilStateCreateInfo, VK_STRUCTURE_TYPE_PIPELINE_DEPTH_STENCIL_STATE_CREATE_INFO);
+
+	VkGraphicsPipelineCreateInfo GraphicsPipelineCreateInfo;
+	ZeroVulkanStruct(GraphicsPipelineCreateInfo, VK_STRUCTURE_TYPE_GRAPHICS_PIPELINE_CREATE_INFO);
+	GraphicsPipelineCreateInfo.sType = VK_STRUCTURE_TYPE_GRAPHICS_PIPELINE_CREATE_INFO;
+	GraphicsPipelineCreateInfo.pNext = &LibraryInfo;
+	GraphicsPipelineCreateInfo.flags = VK_PIPELINE_CREATE_LIBRARY_BIT_KHR | VK_PIPELINE_CREATE_DESCRIPTOR_BUFFER_BIT_EXT |
+		(GVulkanGraphicPipelineLibraryLinkingMode == 1 ? VK_PIPELINE_CREATE_RETAIN_LINK_TIME_OPTIMIZATION_INFO_BIT_EXT : 0);
+	GraphicsPipelineCreateInfo.stageCount = 1;
+	GraphicsPipelineCreateInfo.pStages = &ShaderStageCreateInfo;
+	GraphicsPipelineCreateInfo.pMultisampleState = &MultisampleStateCreateInfo;
+	GraphicsPipelineCreateInfo.pDepthStencilState = &DepthStencilStateCreateInfo;
+	GraphicsPipelineCreateInfo.layout = Device->GetBindlessDescriptorManager()->GetPipelineLayout();
+	GraphicsPipelineCreateInfo.pDynamicState = &DynamicStateCreateInfo;
+
+	VkPipeline PipelineLibrary = VK_NULL_HANDLE;
+	VERIFYVULKANRESULT(VulkanRHI::vkCreateGraphicsPipelines(Device->GetHandle(), VK_NULL_HANDLE, 1, &GraphicsPipelineCreateInfo, VULKAN_CPU_ALLOCATOR, &PipelineLibrary));
+	return PipelineLibrary;
+}
+
 template <typename ShaderType> 
 ShaderType* FVulkanShaderFactory::CreateShader(TArrayView<const uint8> Code, FVulkanDevice* Device)
 {
@@ -88,6 +245,22 @@ ShaderType* FVulkanShaderFactory::CreateShader(TArrayView<const uint8> Code, FVu
 
 				FShaderCodeReader ShaderCode(Code);
 				ReadShaderOptionalData(ShaderCode, RetShader);
+
+				if (Device->SupportsGraphicPipelineLibraries())
+				{
+					if (std::is_same_v<ShaderType, FVulkanVertexShader>)
+					{
+						RetShader->PipelineLibrary = CreatePreRasterShaderLibrary(Device, RetShader, VK_SHADER_STAGE_VERTEX_BIT);
+					}
+					if (std::is_same_v<ShaderType, FVulkanMeshShader>)
+					{
+						RetShader->PipelineLibrary = CreatePreRasterShaderLibrary(Device, RetShader, VK_SHADER_STAGE_MESH_BIT_EXT);
+					}
+					if (std::is_same_v<ShaderType, FVulkanPixelShader>)
+					{
+						RetShader->PipelineLibrary = CreatePixelShaderLibrary(Device, RetShader);
+					}
+				}
 			}
 		}
 	}
@@ -246,6 +419,7 @@ FVulkanDevice* FVulkanShaderModule::Device = nullptr;
 
 FVulkanShaderModule::~FVulkanShaderModule()
 {
+	DEC_DWORD_STAT(STAT_VulkanNumShaderModule);
 	Device->GetDeferredDeletionQueue().EnqueueResource(VulkanRHI::FDeferredDeletionQueue2::EType::ShaderModule, ActualShaderModule);
 }
 
@@ -309,6 +483,12 @@ FVulkanShader::FVulkanShader(FVulkanDevice* InDevice, EShaderFrequency InFrequen
 
 static TRefCountPtr<FVulkanShaderModule> CreateShaderModule(FVulkanDevice* Device, FVulkanShader::FSpirvCode& SpirvCode)
 {
+	INC_DWORD_STAT(STAT_VulkanNumShaderModule);
+
+#if VULKAN_ENABLE_AGGRESSIVE_STATS
+	TRACE_CPUPROFILER_EVENT_SCOPE(VulkanShaderModuleTime);
+#endif
+
 	const TArrayView<uint32> Spirv = SpirvCode.GetCodeView();
 	VkShaderModule ShaderModule;
 	VkShaderModuleCreateInfo ModuleCreateInfo;
@@ -552,6 +732,25 @@ FVulkanShader::~FVulkanShader()
 {
 	PurgeShaderModules();
 	Device->GetShaderFactory().OnDeleteShader(*this);
+
+	if (PipelineLibrary)
+	{
+		Device->GetDeferredDeletionQueue().EnqueueResource(VulkanRHI::FDeferredDeletionQueue2::EType::Pipeline, PipelineLibrary);
+		PipelineLibrary = VK_NULL_HANDLE;
+
+		if (Frequency == SF_Vertex)
+		{
+			DEC_DWORD_STAT(STAT_VulkanNumVSLibs);
+		}
+		else if (Frequency == SF_Mesh)
+		{
+			DEC_DWORD_STAT(STAT_VulkanNumMSLibs);
+		}
+		else if (Frequency == SF_Pixel)
+		{
+			DEC_DWORD_STAT(STAT_VulkanNumPSLibs);
+		}
+	}
 }
 
 void FVulkanShader::PurgeShaderModules()
diff --git a/Engine/Source/Runtime/VulkanRHI/Private/VulkanUtil.cpp b/Engine/Source/Runtime/VulkanRHI/Private/VulkanUtil.cpp
--- a/Engine/Source/Runtime/VulkanRHI/Private/VulkanUtil.cpp
+++ b/Engine/Source/Runtime/VulkanRHI/Private/VulkanUtil.cpp
@@ -1072,7 +1072,7 @@ void FVulkanDynamicRHI::TerminateOnGPUCrash()
 
 
 
-
+// PSO Stats
 DEFINE_STAT(STAT_VulkanNumPSOs);
 DEFINE_STAT(STAT_VulkanNumGraphicsPSOs);
 DEFINE_STAT(STAT_VulkanNumPSOLRU);
@@ -1085,6 +1085,18 @@ DEFINE_STAT(STAT_VulkanNumComputePSOs);
 DEFINE_STAT(STAT_VulkanPSOKeyMemory);
 
 
+// Graphic pipeline library stats
+DEFINE_STAT(STAT_VulkanRegPSOCompileCount);
+DEFINE_STAT(STAT_VulkanLibPSOCompileCount);
+DEFINE_STAT(STAT_VulkanNumVertexInputLibs);
+DEFINE_STAT(STAT_VulkanNumVSLibs);
+DEFINE_STAT(STAT_VulkanNumMSLibs);
+DEFINE_STAT(STAT_VulkanNumPSLibs);
+DEFINE_STAT(STAT_VulkanNumFragOutputStateLibs);
+DEFINE_STAT(STAT_VulkanNumShaderModule);
+
+
+// General Stats
 DEFINE_STAT(STAT_VulkanDrawCallTime);
 DEFINE_STAT(STAT_VulkanDispatchCallTime);
 DEFINE_STAT(STAT_VulkanDrawCallPrepareTime);
