Making VulkanRHI also use -ForceAsyncCompute to enable async compute.

diff --git a/Engine/Source/Runtime/VulkanRHI/Private/VulkanDevice.cpp b/Engine/Source/Runtime/VulkanRHI/Private/VulkanDevice.cpp
--- a/Engine/Source/Runtime/VulkanRHI/Private/VulkanDevice.cpp
+++ b/Engine/Source/Runtime/VulkanRHI/Private/VulkanDevice.cpp
@@ -357,7 +357,7 @@ void FVulkanDevice::CreateDevice(TArray<const ANSICHAR*>& DeviceLayers, FVulkanD
 			// - a new family index is available
 			// - Sync2 is available
 			if ((ComputeQueueFamilyIndex == -1) &&
-				(GRHIAllowAsyncComputeCvar.GetValueOnAnyThread() != 0 || GAllowPresentOnComputeQueue.GetValueOnAnyThread() != 0) && 
+				(GRHIAllowAsyncComputeCvar.GetValueOnAnyThread() != 0 || GAllowPresentOnComputeQueue.GetValueOnAnyThread() != 0 || FParse::Param(FCommandLine::Get(), TEXT("ForceAsyncCompute"))) &&
 				(GfxQueueFamilyIndex != FamilyIndex) && 
 				SupportsParallelRendering())
 			{
