Removing immediate context usage in Vulkan descriptor updates by passing Contexts down through the view update methods.

diff --git a/Engine/Source/Runtime/VulkanRHI/Private/VulkanBindlessDescriptorManager.cpp b/Engine/Source/Runtime/VulkanRHI/Private/VulkanBindlessDescriptorManager.cpp
--- a/Engine/Source/Runtime/VulkanRHI/Private/VulkanBindlessDescriptorManager.cpp
+++ b/Engine/Source/Runtime/VulkanRHI/Private/VulkanBindlessDescriptorManager.cpp
@@ -736,7 +736,7 @@ FRHIDescriptorHandle FVulkanBindlessDescriptorManager::AllocateDescriptor(VkDesc
 	if (IsSupported())
 	{
 		FRHIDescriptorHandle DescriptorHandle = AllocateDescriptor(DescriptorType);
-		UpdateBuffer(DescriptorHandle, DescriptorType, BufferAddress, BufferSize, true);
+		UpdateBuffer({}, DescriptorHandle, DescriptorType, BufferAddress, BufferSize, true);
 		return DescriptorHandle;
 	}
 
@@ -748,7 +748,7 @@ FRHIDescriptorHandle FVulkanBindlessDescriptorManager::AllocateDescriptor(VkDesc
 	if (IsSupported())
 	{
 		FRHIDescriptorHandle DescriptorHandle = AllocateDescriptor(DescriptorType);
-		UpdateBuffer(DescriptorHandle, DescriptorType, Buffer->GetHandle(), Buffer->GetOffset() + ExtraOffset, Buffer->GetCurrentSize() - ExtraOffset, true);
+		UpdateBuffer({}, DescriptorHandle, DescriptorType, Buffer->GetHandle(), Buffer->GetOffset() + ExtraOffset, Buffer->GetCurrentSize() - ExtraOffset, true);
 		return DescriptorHandle;
 	}
 
@@ -760,7 +760,7 @@ FRHIDescriptorHandle FVulkanBindlessDescriptorManager::AllocateDescriptor(VkSamp
 	if (IsSupported())
 	{
 		FRHIDescriptorHandle DescriptorHandle = AllocateDescriptor(VK_DESCRIPTOR_TYPE_SAMPLER);
-		UpdateSampler(DescriptorHandle, VulkanSampler);
+		UpdateSampler({}, DescriptorHandle, VulkanSampler);
 		return DescriptorHandle;
 	}
 
@@ -779,7 +779,7 @@ void FVulkanBindlessDescriptorManager::FreeDescriptor(FRHIDescriptorHandle Descr
 	}
 }
 
-void FVulkanBindlessDescriptorManager::UpdateDescriptor(FRHIDescriptorHandle DescriptorHandle, VkDescriptorType DescriptorType, VkDescriptorDataEXT DescriptorData, bool bImmediateUpdate)
+void FVulkanBindlessDescriptorManager::UpdateDescriptor(const FVulkanContextArray& Contexts, FRHIDescriptorHandle DescriptorHandle, VkDescriptorType DescriptorType, VkDescriptorDataEXT DescriptorData, bool bImmediateUpdate)
 {
 	checkf(DescriptorHandle.IsValid(), TEXT("Attemping to update invalid descriptor handle!"));
 
@@ -801,19 +801,20 @@ void FVulkanBindlessDescriptorManager::UpdateDescriptor(FRHIDescriptorHandle Des
 	}
 	else
 	{
-		check(!IsInRenderingThread() || FRHICommandListExecutor::GetImmediateCommandList().Bypass() || !IsRunningRHIInSeparateThread());
-		FVulkanCommandBuffer* CmdBuffer = Device.GetImmediateContext().GetActiveCmdBuffer();
-
-		// :todo-jn:  Hack to avoid barriers/copies in renderpasses
-		if (CmdBuffer->IsInsideRenderPass())
+		if (FVulkanCommandListContext * Context = Contexts[ERHIPipeline::Graphics])
 		{
-			FMemory::Memcpy(&State.MappedPointer[ByteOffset], CPUDescriptor, State.DescriptorSize);
-		}
-		else
-		{
-			VulkanRHI::FStagingBuffer* StagingBuffer = Device.GetStagingManager().AcquireBuffer(State.DescriptorSize);
-			FMemory::Memcpy(StagingBuffer->GetMappedPointer(), CPUDescriptor, State.DescriptorSize);
+			FVulkanCommandBuffer* CmdBuffer = Context->GetActiveCmdBuffer();
+
+			// :todo-jn:  Hack to avoid barriers/copies in renderpasses
+			if (CmdBuffer->IsInsideRenderPass())
 			{
+				FMemory::Memcpy(&State.MappedPointer[ByteOffset], CPUDescriptor, State.DescriptorSize);
+			}
+			else
+			{
+				VulkanRHI::FStagingBuffer* StagingBuffer = Device.GetStagingManager().AcquireBuffer(State.DescriptorSize);
+				FMemory::Memcpy(StagingBuffer->GetMappedPointer(), CPUDescriptor, State.DescriptorSize);
+
 				VkMemoryBarrier2 MemoryBarrier;
 				ZeroVulkanStruct(MemoryBarrier, VK_STRUCTURE_TYPE_MEMORY_BARRIER_2);
 				MemoryBarrier.srcStageMask = VK_PIPELINE_STAGE_2_ALL_COMMANDS_BIT;
@@ -838,23 +839,24 @@ void FVulkanBindlessDescriptorManager::UpdateDescriptor(FRHIDescriptorHandle Des
 				MemoryBarrier.dstStageMask = VK_PIPELINE_STAGE_2_ALL_COMMANDS_BIT;
 				MemoryBarrier.dstAccessMask = VK_ACCESS_2_DESCRIPTOR_BUFFER_READ_BIT_EXT;
 				VulkanRHI::vkCmdPipelineBarrier2KHR(CmdBuffer->GetHandle(), &DependencyInfo);
+
+				Device.GetStagingManager().ReleaseBuffer(Context, StagingBuffer);
 			}
-			Device.GetStagingManager().ReleaseBuffer(&Device.GetImmediateContext(), StagingBuffer);
 		}
 	}
 }
 
-void FVulkanBindlessDescriptorManager::UpdateSampler(FRHIDescriptorHandle DescriptorHandle, VkSampler VulkanSampler)
+void FVulkanBindlessDescriptorManager::UpdateSampler(const FVulkanContextArray& Contexts, FRHIDescriptorHandle DescriptorHandle, VkSampler VulkanSampler)
 {
 	if (bIsSupported)
 	{
 		VkDescriptorDataEXT DescriptorData;
 		DescriptorData.pSampler = &VulkanSampler;
-		UpdateDescriptor(DescriptorHandle, VK_DESCRIPTOR_TYPE_SAMPLER, DescriptorData, true);
+		UpdateDescriptor(Contexts, DescriptorHandle, VK_DESCRIPTOR_TYPE_SAMPLER, DescriptorData, true);
 	}
 }
 
-void FVulkanBindlessDescriptorManager::UpdateImage(FRHIDescriptorHandle DescriptorHandle, VkDescriptorType DescriptorType, VkImageView ImageView, bool bIsDepthStencil, bool bImmediateUpdate)
+void FVulkanBindlessDescriptorManager::UpdateImage(const FVulkanContextArray& Contexts, FRHIDescriptorHandle DescriptorHandle, VkDescriptorType DescriptorType, VkImageView ImageView, bool bIsDepthStencil, bool bImmediateUpdate)
 {
 	if (bIsSupported)
 	{
@@ -868,11 +870,11 @@ void FVulkanBindlessDescriptorManager::UpdateImage(FRHIDescriptorHandle Descript
 
 		VkDescriptorDataEXT DescriptorData;
 		DescriptorData.pSampledImage = &DescriptorImageInfo;  // same pointer for storage, it's a union
-		UpdateDescriptor(DescriptorHandle, DescriptorType, DescriptorData, bImmediateUpdate);
+		UpdateDescriptor(Contexts, DescriptorHandle, DescriptorType, DescriptorData, bImmediateUpdate);
 	}
 }
 
-void FVulkanBindlessDescriptorManager::UpdateBuffer(FRHIDescriptorHandle DescriptorHandle, VkDescriptorType DescriptorType, VkBuffer VulkanBuffer, VkDeviceSize BufferOffset, VkDeviceSize BufferSize, bool bImmediateUpdate)
+void FVulkanBindlessDescriptorManager::UpdateBuffer(const FVulkanContextArray& Contexts, FRHIDescriptorHandle DescriptorHandle, VkDescriptorType DescriptorType, VkBuffer VulkanBuffer, VkDeviceSize BufferOffset, VkDeviceSize BufferSize, bool bImmediateUpdate)
 {
 	if (bIsSupported)
 	{
@@ -883,11 +885,11 @@ void FVulkanBindlessDescriptorManager::UpdateBuffer(FRHIDescriptorHandle Descrip
 		BufferInfo.buffer = VulkanBuffer;
 		const VkDeviceAddress BufferAddress = VulkanRHI::vkGetBufferDeviceAddressKHR(Device.GetHandle(), &BufferInfo);
 
-		UpdateBuffer(DescriptorHandle, DescriptorType, BufferAddress + BufferOffset, BufferSize, bImmediateUpdate);
+		UpdateBuffer(Contexts, DescriptorHandle, DescriptorType, BufferAddress + BufferOffset, BufferSize, bImmediateUpdate);
 	}
 }
 
-void FVulkanBindlessDescriptorManager::UpdateBuffer(FRHIDescriptorHandle DescriptorHandle, VkDescriptorType DescriptorType, VkDeviceAddress BufferAddress, VkDeviceSize BufferSize, bool bImmediateUpdate)
+void FVulkanBindlessDescriptorManager::UpdateBuffer(const FVulkanContextArray& Contexts, FRHIDescriptorHandle DescriptorHandle, VkDescriptorType DescriptorType, VkDeviceAddress BufferAddress, VkDeviceSize BufferSize, bool bImmediateUpdate)
 {
 	if (bIsSupported)
 	{
@@ -900,11 +902,11 @@ void FVulkanBindlessDescriptorManager::UpdateBuffer(FRHIDescriptorHandle Descrip
 
 		VkDescriptorDataEXT DescriptorData;
 		DescriptorData.pStorageBuffer = &AddressInfo;  // same pointer for uniform, it's a union
-		UpdateDescriptor(DescriptorHandle, DescriptorType, DescriptorData, bImmediateUpdate);
+		UpdateDescriptor(Contexts, DescriptorHandle, DescriptorType, DescriptorData, bImmediateUpdate);
 	}
 }
 
-void FVulkanBindlessDescriptorManager::UpdateTexelBuffer(FRHIDescriptorHandle DescriptorHandle, VkDescriptorType DescriptorType, const VkBufferViewCreateInfo& ViewInfo, bool bImmediateUpdate)
+void FVulkanBindlessDescriptorManager::UpdateTexelBuffer(const FVulkanContextArray& Contexts, FRHIDescriptorHandle DescriptorHandle, VkDescriptorType DescriptorType, const VkBufferViewCreateInfo& ViewInfo, bool bImmediateUpdate)
 {
 	if (bIsSupported)
 	{
@@ -924,11 +926,11 @@ void FVulkanBindlessDescriptorManager::UpdateTexelBuffer(FRHIDescriptorHandle De
 
 		VkDescriptorDataEXT DescriptorData;
 		DescriptorData.pUniformTexelBuffer = &AddressInfo;  // same pointer for storage, it's a union
-		UpdateDescriptor(DescriptorHandle, DescriptorType, DescriptorData, bImmediateUpdate);
+		UpdateDescriptor(Contexts, DescriptorHandle, DescriptorType, DescriptorData, bImmediateUpdate);
 	}
 }
 
-void FVulkanBindlessDescriptorManager::UpdateAccelerationStructure(FRHIDescriptorHandle DescriptorHandle, VkAccelerationStructureKHR AccelerationStructure, bool bImmediateUpdate)
+void FVulkanBindlessDescriptorManager::UpdateAccelerationStructure(const FVulkanContextArray& Contexts, FRHIDescriptorHandle DescriptorHandle, VkAccelerationStructureKHR AccelerationStructure, bool bImmediateUpdate)
 {
 	if (bIsSupported)
 	{
@@ -940,7 +942,7 @@ void FVulkanBindlessDescriptorManager::UpdateAccelerationStructure(FRHIDescripto
 
 		VkDescriptorDataEXT DescriptorData;
 		DescriptorData.accelerationStructure = BufferAddress;
-		UpdateDescriptor(DescriptorHandle, VK_DESCRIPTOR_TYPE_ACCELERATION_STRUCTURE_KHR, DescriptorData, bImmediateUpdate);
+		UpdateDescriptor(Contexts, DescriptorHandle, VK_DESCRIPTOR_TYPE_ACCELERATION_STRUCTURE_KHR, DescriptorData, bImmediateUpdate);
 	}
 }
 
diff --git a/Engine/Source/Runtime/VulkanRHI/Private/VulkanBindlessDescriptorManager.h b/Engine/Source/Runtime/VulkanRHI/Private/VulkanBindlessDescriptorManager.h
--- a/Engine/Source/Runtime/VulkanRHI/Private/VulkanBindlessDescriptorManager.h
+++ b/Engine/Source/Runtime/VulkanRHI/Private/VulkanBindlessDescriptorManager.h
@@ -45,14 +45,14 @@ public:
 
 	void FreeDescriptor(FRHIDescriptorHandle DescriptorHandle);
 
-	void UpdateSampler(FRHIDescriptorHandle DescriptorHandle, VkSampler VulkanSampler);
-	void UpdateImage(FRHIDescriptorHandle DescriptorHandle, VkDescriptorType DescriptorType, VkImageView VulkanImage, bool bIsDepthStencil, bool bImmediateUpdate = true);
+	void UpdateSampler(const FVulkanContextArray& Contexts, FRHIDescriptorHandle DescriptorHandle, VkSampler VulkanSampler);
+	void UpdateImage(const FVulkanContextArray& Contexts, FRHIDescriptorHandle DescriptorHandle, VkDescriptorType DescriptorType, VkImageView VulkanImage, bool bIsDepthStencil, bool bImmediateUpdate = true);
 
-	void UpdateBuffer(FRHIDescriptorHandle DescriptorHandle, VkDescriptorType DescriptorType, VkBuffer VulkanBuffer, VkDeviceSize BufferOffset, VkDeviceSize BufferSize, bool bImmediateUpdate = true);
-	void UpdateBuffer(FRHIDescriptorHandle DescriptorHandle, VkDescriptorType DescriptorType, VkDeviceAddress BufferAddress, VkDeviceSize BufferSize, bool bImmediateUpdate = true);
+	void UpdateBuffer(const FVulkanContextArray& Contexts, FRHIDescriptorHandle DescriptorHandle, VkDescriptorType DescriptorType, VkBuffer VulkanBuffer, VkDeviceSize BufferOffset, VkDeviceSize BufferSize, bool bImmediateUpdate = true);
+	void UpdateBuffer(const FVulkanContextArray& Contexts, FRHIDescriptorHandle DescriptorHandle, VkDescriptorType DescriptorType, VkDeviceAddress BufferAddress, VkDeviceSize BufferSize, bool bImmediateUpdate = true);
 
-	void UpdateTexelBuffer(FRHIDescriptorHandle DescriptorHandle, VkDescriptorType DescriptorType, const VkBufferViewCreateInfo& ViewInfo, bool bImmediateUpdate = true);
-	void UpdateAccelerationStructure(FRHIDescriptorHandle DescriptorHandle, VkAccelerationStructureKHR AccelerationStructure, bool bImmediateUpdate = true);
+	void UpdateTexelBuffer(const FVulkanContextArray& Contexts, FRHIDescriptorHandle DescriptorHandle, VkDescriptorType DescriptorType, const VkBufferViewCreateInfo& ViewInfo, bool bImmediateUpdate = true);
+	void UpdateAccelerationStructure(const FVulkanContextArray& Contexts, FRHIDescriptorHandle DescriptorHandle, VkAccelerationStructureKHR AccelerationStructure, bool bImmediateUpdate = true);
 
 	void RegisterUniformBuffers(FVulkanCommandListContext& Context, VkPipelineBindPoint BindPoint, const FUniformBufferDescriptorArrays& StageUBs);
 
@@ -101,6 +101,6 @@ private:
 	BindlessSetState& GetBindlessState(ERHIDescriptorType DescriptorType);
 	BindlessSetState& GetBindlessState(VkDescriptorType DescriptorType);
 
-	void UpdateDescriptor(FRHIDescriptorHandle DescriptorHandle, VkDescriptorType DescriptorType, VkDescriptorDataEXT DescriptorData, bool bImmediateUpdate);
+	void UpdateDescriptor(const FVulkanContextArray& Contexts, FRHIDescriptorHandle DescriptorHandle, VkDescriptorType DescriptorType, VkDescriptorDataEXT DescriptorData, bool bImmediateUpdate);
 	static bool VerifySupport(FVulkanDevice& InDevice);
 };
diff --git a/Engine/Source/Runtime/VulkanRHI/Private/VulkanBuffer.cpp b/Engine/Source/Runtime/VulkanRHI/Private/VulkanBuffer.cpp
--- a/Engine/Source/Runtime/VulkanRHI/Private/VulkanBuffer.cpp
+++ b/Engine/Source/Runtime/VulkanRHI/Private/VulkanBuffer.cpp
@@ -298,19 +298,15 @@ void* FVulkanBuffer::Lock(FRHICommandListBase& RHICmdList, EResourceLockMode Loc
 			AllocateMemory(NewAlloc);
 			NewAlloc.Alloc.Disown();
 
-			RHICmdList.EnqueueLambda(TEXT("FVulkanBuffer::Lock"), [Buffer = this, NewAlloc](FRHICommandListBase& CmdList)
-			{
-				UpdateVulkanBufferStats(Buffer->GetDesc(), Buffer->CurrentBufferAlloc.Alloc.Size, false);
-				Buffer->CurrentBufferAlloc.Alloc.Free(Buffer->Device);
-				Buffer->CurrentBufferAlloc = NewAlloc;
-				Buffer->CurrentBufferAlloc.Alloc.Own();
-				Buffer->UpdateLinkedViews();
-			});
-
-			if (RHICmdList.IsTopOfPipe())
-			{
-				RHICmdList.RHIThreadFence(true);
-			}
+			RHICmdList.EnqueueLambdaMultiPipe(GetEnabledRHIPipelines(), FRHICommandListBase::EThreadFence::Enabled, TEXT("FVulkanBuffer::Lock"),
+				[Buffer = this, NewAlloc](const FVulkanContextArray& Contexts)
+				{
+					UpdateVulkanBufferStats(Buffer->GetDesc(), Buffer->CurrentBufferAlloc.Alloc.Size, false);
+					Buffer->CurrentBufferAlloc.Alloc.Free(Buffer->Device);
+					Buffer->CurrentBufferAlloc = NewAlloc;
+					Buffer->CurrentBufferAlloc.Alloc.Own();
+					Buffer->UpdateLinkedViews(Contexts);
+				});
 
 			Data = NewAlloc.HostPtr;
 			DataOffset = Offset;
diff --git a/Engine/Source/Runtime/VulkanRHI/Private/VulkanContext.h b/Engine/Source/Runtime/VulkanRHI/Private/VulkanContext.h
--- a/Engine/Source/Runtime/VulkanRHI/Private/VulkanContext.h
+++ b/Engine/Source/Runtime/VulkanRHI/Private/VulkanContext.h
@@ -512,6 +512,11 @@ struct TVulkanResourceTraits<IRHIUploadContext>
 
 struct FVulkanContextArray : public TRHIPipelineArray<FVulkanCommandListContext*>
 {
+	FVulkanContextArray()
+		: TRHIPipelineArray(InPlace, nullptr)
+	{
+	}
+
 	FVulkanContextArray(FRHIContextArray const& Contexts)
 		: TRHIPipelineArray(InPlace, nullptr)
 	{
diff --git a/Engine/Source/Runtime/VulkanRHI/Private/VulkanMemory.cpp b/Engine/Source/Runtime/VulkanRHI/Private/VulkanMemory.cpp
--- a/Engine/Source/Runtime/VulkanRHI/Private/VulkanMemory.cpp
+++ b/Engine/Source/Runtime/VulkanRHI/Private/VulkanMemory.cpp
@@ -1683,7 +1683,7 @@ namespace VulkanRHI
 		ReleasePage(InPage);
 	}
 
-	uint64 FVulkanResourceHeap::EvictOne(FVulkanDevice& Device, FVulkanCommandListContext& Context)
+	uint64 FVulkanResourceHeap::EvictOne(FVulkanDevice& Device, const FVulkanContextArray& Contexts)
 	{
 		FScopeLock ScopeLock(&PagesLock);
 		for(int32 Index = MAX_BUCKETS - 2; Index >= 0; Index--)
@@ -1694,7 +1694,7 @@ namespace VulkanRHI
 				FVulkanSubresourceAllocator* Allocator = Pages[Index2];
 				if (!Allocator->bIsEvicting && Allocator->GetSubresourceAllocatorFlags() & VulkanAllocationFlagsCanEvict)
 				{
-					return Allocator->EvictToHost(Device, Context);
+					return Allocator->EvictToHost(Device, Contexts);
 				}
 			}
 		}
@@ -1740,7 +1740,7 @@ namespace VulkanRHI
 		return Allocator->GetIsDefragging();
 	}
 
-	void FVulkanResourceHeap::DefragTick(FVulkanDevice& Device, FVulkanCommandListContext& Context, uint32 Count)
+	void FVulkanResourceHeap::DefragTick(FVulkanDevice& Device, const FVulkanContextArray& Contexts, uint32 Count)
 	{
 		SCOPED_NAMED_EVENT(FVulkanSubresourceAllocator_DefragTick, FColor::Cyan);
 
@@ -1834,7 +1834,7 @@ namespace VulkanRHI
 			{
 				VULKAN_LOGMEMORY(TEXT("Defragging heap [%d] Max %6.2fMB Free %6.2fMB   %2d"), CurrentDefragTarget->AllocatorIndex, MaxSize / (1024.f*1024.f), FreeSpace / (1024.f * 1024.f), DefragCountDown);
 			}
-			if (CurrentDefragTarget->DefragTick(Device, Context, this, Count) > 0)
+			if (CurrentDefragTarget->DefragTick(Device, Contexts, this, Count) > 0)
 			{
 				DefragCountDown = 3;
 			}
@@ -2267,7 +2267,7 @@ namespace VulkanRHI
 		DeviceMemoryManager->TrimMemory(bImmediately);
 	}
 
-	void FMemoryManager::ReleaseFreedPages(FVulkanCommandListContext& Context)
+	void FMemoryManager::ReleaseFreedPages(const FVulkanContextArray& Contexts)
 	{
 		auto CanDefragHeap = [](FVulkanResourceHeap* Heap)
 		{
@@ -2343,7 +2343,7 @@ namespace VulkanRHI
 		if(BestEvictHeap && ((GVulkanEvictOnePage || UpdateEvictThreshold(true)) && PrimaryHeapIndex >= 0))
 		{
 			GVulkanEvictOnePage = 0;
-			PendingEvictBytes += BestEvictHeap->EvictOne(Device, Context);
+			PendingEvictBytes += BestEvictHeap->EvictOne(Device, Contexts);
 		}
 
 #if !PLATFORM_ANDROID
@@ -2359,7 +2359,7 @@ namespace VulkanRHI
 			{
 				DumpMemory();
 			}
-			BestDefragHeap->DefragTick(Device, Context, Count);
+			BestDefragHeap->DefragTick(Device, Contexts, Count);
 
 			if(GVulkanDefragOnce)
 			{
@@ -4130,7 +4130,7 @@ namespace VulkanRHI
 	// One tick of incremental defrag. This will move max Count_ allocations from this suballocator to other SubAllocators on the same Heap
 	// This will eventually free up the page, which can then be released
 	// Returns # of allocations moved.
-	int32 FVulkanSubresourceAllocator::DefragTick(FVulkanDevice& Device, FVulkanCommandListContext& Context, FVulkanResourceHeap* Heap, uint32 Count_)
+	int32 FVulkanSubresourceAllocator::DefragTick(FVulkanDevice& Device, const FVulkanContextArray& Contexts, FVulkanResourceHeap* Heap, uint32 Count_)
 	{
 		LastDefragFrame = GFrameNumberRenderThread;
 		int32 DefragCount = 0;
@@ -4190,7 +4190,7 @@ namespace VulkanRHI
 
 								//Move the Rendertarget to the new allocation
 								//Function swaps the old allocation into the Allocation object
-								Texture->Move(Device, Context, Allocation);
+								Texture->Move(Device, Contexts, Allocation);
 								DefragCount++;
 								Device.GetMemoryManager().FreeVulkanAllocation(Allocation);
 
@@ -4219,7 +4219,7 @@ namespace VulkanRHI
 		return DefragCount;
 	}
 
-	uint64 FVulkanSubresourceAllocator::EvictToHost(FVulkanDevice& Device, FVulkanCommandListContext& Context)
+	uint64 FVulkanSubresourceAllocator::EvictToHost(FVulkanDevice& Device, const FVulkanContextArray& Contexts)
 	{
 		FScopeLock ScopeLock(&SubresourceAllocatorCS);
 		bIsEvicting = true;
@@ -4234,7 +4234,7 @@ namespace VulkanRHI
 					FVulkanEvictable* Texture = Alloc.AllocationOwner;
 					if (Texture->CanEvict())
 					{
-						Texture->Evict(Device, Context);
+						Texture->Evict(Device, Contexts);
 					}
 				}
 				break;
diff --git a/Engine/Source/Runtime/VulkanRHI/Private/VulkanMemory.h b/Engine/Source/Runtime/VulkanRHI/Private/VulkanMemory.h
--- a/Engine/Source/Runtime/VulkanRHI/Private/VulkanMemory.h
+++ b/Engine/Source/Runtime/VulkanRHI/Private/VulkanMemory.h
@@ -33,6 +33,8 @@ class FVulkanQueue;
 class FVulkanCommandBuffer;
 class FVulkanTexture;
 
+struct FVulkanContextArray;
+
 using FVulkanSyncPoint = FGraphEvent;
 using FVulkanSyncPointRef = TRefCountPtr<FVulkanSyncPoint>;
 
@@ -102,8 +104,8 @@ class FVulkanEvictable
 public:
 	virtual bool CanMove() const { return false; }
 	virtual bool CanEvict() const { return false; }
-	virtual void Evict(FVulkanDevice& Device, FVulkanCommandListContext& Context) { checkNoEntry(); }
-	virtual void Move(FVulkanDevice& Device, FVulkanCommandListContext& Context, VulkanRHI::FVulkanAllocation& Allocation) { checkNoEntry(); }
+	virtual void Evict(FVulkanDevice& Device, const FVulkanContextArray& Contexts) { checkNoEntry(); }
+	virtual void Move(FVulkanDevice& Device, const FVulkanContextArray& Contexts, VulkanRHI::FVulkanAllocation& Allocation) { checkNoEntry(); }
 	virtual FVulkanTexture* GetEvictableTexture() { return nullptr; }
 };
 
@@ -598,9 +600,9 @@ namespace VulkanRHI
 	protected:
 		void SetIsDefragging(bool bInIsDefragging){ bIsDefragging = bInIsDefragging; }
 		void DumpFullHeap();
-		int32 DefragTick(FVulkanDevice& Device, FVulkanCommandListContext& Context, FVulkanResourceHeap* Heap, uint32 Count);
+		int32 DefragTick(FVulkanDevice& Device, const FVulkanContextArray& Contexts, FVulkanResourceHeap* Heap, uint32 Count);
 		bool CanDefrag();
-		uint64 EvictToHost(FVulkanDevice& Device, FVulkanCommandListContext& Context);
+		uint64 EvictToHost(FVulkanDevice& Device, const FVulkanContextArray& Contexts);
 		void SetFreePending(FVulkanAllocation& Allocation);
 		void FreeInternalData(int32 Index);
 		int32 AllocateInternalData();
@@ -683,8 +685,8 @@ namespace VulkanRHI
 			return MemoryTypeIndex;
 		}
 
-		uint64 EvictOne(FVulkanDevice& Device, FVulkanCommandListContext& Context);
-		void DefragTick(FVulkanDevice& Device, FVulkanCommandListContext& Context, uint32 Count);
+		uint64 EvictOne(FVulkanDevice& Device, const FVulkanContextArray& Contexts);
+		void DefragTick(FVulkanDevice& Device, const FVulkanContextArray& Contexts, uint32 Count);
 		void DumpMemory(FResourceHeapStats& Stats);
 
 		void SetDefragging(FVulkanSubresourceAllocator* Allocator);
@@ -790,7 +792,7 @@ namespace VulkanRHI
 		bool ReleaseSubresourceAllocator(FVulkanSubresourceAllocator* SubresourceAllocator);
 
 
-		void ReleaseFreedPages(FVulkanCommandListContext& Context);
+		void ReleaseFreedPages(const FVulkanContextArray& Contexts);
 		void DumpMemory(bool bFullDump = true);
 
 
diff --git a/Engine/Source/Runtime/VulkanRHI/Private/VulkanRHI.cpp b/Engine/Source/Runtime/VulkanRHI/Private/VulkanRHI.cpp
--- a/Engine/Source/Runtime/VulkanRHI/Private/VulkanRHI.cpp
+++ b/Engine/Source/Runtime/VulkanRHI/Private/VulkanRHI.cpp
@@ -1134,7 +1134,7 @@ void FVulkanDynamicRHI::RHIEndFrame_RenderThread(FRHICommandListImmediate& RHICm
 		Context.FreeUnusedCmdBuffers(bTrimMemory);
 
 		Context.Device.GetStagingManager().ProcessPendingFree(false, true);
-		Context.Device.GetMemoryManager().ReleaseFreedPages(Context);
+		Context.Device.GetMemoryManager().ReleaseFreedPages(Contexts);
 		Context.Device.GetDeferredDeletionQueue().ReleaseResources();
 
 		if (UseVulkanDescriptorCache())
@@ -2074,8 +2074,8 @@ bool FVulkanDynamicRHI::RHIMatchPrecachePSOInitializers(const FGraphicsPipelineS
 
 void FVulkanDynamicRHI::RHIReplaceResources(FRHICommandListBase& RHICmdList, TArray<FRHIResourceReplaceInfo>&& ReplaceInfos)
 {
-	RHICmdList.EnqueueLambda(TEXT("FVulkanDynamicRHI::RHIReplaceResources"),
-		[ReplaceInfos = MoveTemp(ReplaceInfos)](FRHICommandListBase& ExecutingCmdList)
+	RHICmdList.EnqueueLambdaMultiPipe(GetEnabledRHIPipelines(), FRHICommandListBase::EThreadFence::Enabled, TEXT("FVulkanDynamicRHI::RHIReplaceResources"),
+		[ReplaceInfos = MoveTemp(ReplaceInfos)](const FVulkanContextArray& Contexts)
 		{
 			for (FRHIResourceReplaceInfo const& Info : ReplaceInfos)
 			{
@@ -2102,7 +2102,7 @@ void FVulkanDynamicRHI::RHIReplaceResources(FRHICommandListBase& RHICmdList, TAr
 							Dst->ReleaseOwnership();
 						}
 
-						Dst->UpdateLinkedViews();
+						Dst->UpdateLinkedViews(Contexts);
 					}
 					break;
 
@@ -2127,8 +2127,6 @@ void FVulkanDynamicRHI::RHIReplaceResources(FRHICommandListBase& RHICmdList, TAr
 			}
 		}
 	);
-
-	RHICmdList.RHIThreadFence(true);
 }
 
 #if PLATFORM_SUPPORTS_BINDLESS_RENDERING
diff --git a/Engine/Source/Runtime/VulkanRHI/Private/VulkanRayTracing.cpp b/Engine/Source/Runtime/VulkanRHI/Private/VulkanRayTracing.cpp
--- a/Engine/Source/Runtime/VulkanRHI/Private/VulkanRayTracing.cpp
+++ b/Engine/Source/Runtime/VulkanRHI/Private/VulkanRayTracing.cpp
@@ -735,17 +735,22 @@ void FVulkanRayTracingScene::BindBuffer(FRHIBuffer* InBuffer, uint32 InBufferOff
 		check(InBufferOffset % GRHIRayTracingAccelerationStructureAlignment == 0);
 
 		View = MakeUnique<FVulkanView>(Device, VK_DESCRIPTOR_TYPE_ACCELERATION_STRUCTURE_KHR);
-		VkAccelerationStructureKHR NativeAccelerationStructureHandle = View->InitAsAccelerationStructureView(
+		View->InitAsAccelerationStructureView(
+			{},
 			AccelerationStructureBuffer
 			, InBufferOffset
 			//, SizeInfo.ResultSize
 			// TODO: Using whole remaining size instead of SizeInfo.ResultSize reintroduces a validation error but use of SizeInfo.ResultSize broke RT on Adreno.
 			, InBuffer->GetSize() - InBufferOffset
-		)->GetAccelerationStructureView().Handle;
+		);
+
+#if VULKAN_USE_DEBUG_NAMES
+		VkAccelerationStructureKHR NativeAccelerationStructureHandle = View->GetAccelerationStructureView().Handle;
 
 		FString DebugNameString = Initializer.DebugName.ToString();
 		DebugNameString = (DebugNameString.IsEmpty()) ? TEXT("TLAS") : DebugNameString;
 		VULKAN_SET_DEBUG_NAME(Device, VK_OBJECT_TYPE_ACCELERATION_STRUCTURE_KHR, NativeAccelerationStructureHandle, TEXT("%s"), *DebugNameString);
+#endif
 	}
 }
 
diff --git a/Engine/Source/Runtime/VulkanRHI/Private/VulkanResources.h b/Engine/Source/Runtime/VulkanRHI/Private/VulkanResources.h
--- a/Engine/Source/Runtime/VulkanRHI/Private/VulkanResources.h
+++ b/Engine/Source/Runtime/VulkanRHI/Private/VulkanResources.h
@@ -33,6 +33,7 @@ class FVulkanShaderResourceView;
 class FVulkanUnorderedAccessView;
 class FVulkanRenderQuery;
 class FVulkanDynamicRHI;
+struct FVulkanContextArray;
 
 using FVulkanSyncPoint = FGraphEvent;
 using FVulkanSyncPointRef = TRefCountPtr<FVulkanSyncPoint>;
@@ -501,18 +502,15 @@ public:
 	FStructuredBufferView      const& GetStructuredBufferView     () const { return Storage.Get<FStructuredBufferView     >(); }
 	FAccelerationStructureView const& GetAccelerationStructureView() const { return Storage.Get<FAccelerationStructureView>(); }
 
-	// NOTE: The InOffset applies to the FVulkanBuffer (it does not include any internal Allocation offsets that may exist)
-	FVulkanView* InitAsTypedBufferView(
-		  FVulkanBuffer* Buffer
-		, EPixelFormat Format
-		, uint32 InOffset
-		, uint32 InSize);
+	void CreateTextureView(VkImage InImage, const FVulkanTextureViewDesc& ViewDesc);
+	void UpdateTextureView(const FVulkanContextArray& Contexts, VkImage InImage, const FVulkanTextureViewDesc& ViewDesc);
 
-	FVulkanView* InitAsTextureView(VkImage InImage, const FVulkanTextureViewDesc& ViewDesc);
+	// NOTE: The InOffset applies to the FVulkanBuffer (it does not include any internal Allocation offsets that may exist)
+	void InitAsTypedBufferView(const FVulkanContextArray& Contexts, FVulkanBuffer* Buffer, EPixelFormat Format, uint32 InOffset, uint32 InSize);
 
 	// NOTE: The InOffset applies to the FVulkanBuffer (it does not include any internal Allocation offsets that may exist)
-	FVulkanView* InitAsStructuredBufferView(FVulkanBuffer* Buffer, uint32 InOffset, uint32 InSize);
-	FVulkanView* InitAsAccelerationStructureView(FVulkanBuffer* Buffer, uint32 Offset, uint32 Size);
+	void InitAsStructuredBufferView(const FVulkanContextArray& Contexts, FVulkanBuffer* Buffer, uint32 InOffset, uint32 InSize);
+	void InitAsAccelerationStructureView(const FVulkanContextArray& Contexts, FVulkanBuffer* Buffer, uint32 Offset, uint32 Size);
 
 	// No moving or copying
 	FVulkanView(FVulkanView     &&) = delete;
@@ -530,6 +528,9 @@ public:
 		return DescriptorType;
 	}
 
+protected:
+	void InitAsTextureView(const FVulkanContextArray& Contexts, VkImage InImage, const FVulkanTextureViewDesc& ViewDesc);
+
 private:
 	FVulkanDevice& Device;
 	FRHIDescriptorHandle BindlessHandle;
@@ -550,7 +551,7 @@ protected:
 	}
 
 public:
-	virtual void UpdateView() = 0;
+	virtual void UpdateView(const FVulkanContextArray& Contexts) = 0;
 };
 
 class FVulkanViewableResource
@@ -568,7 +569,7 @@ public:
 
 	// @todo convert views owned by the texture into proper
 	// FVulkanView instances, then remove 'virtual' from this class
-	virtual void UpdateLinkedViews();
+	virtual void UpdateLinkedViews(const FVulkanContextArray& Contexts);
 
 private:
 	friend FVulkanShaderResourceView;
@@ -628,8 +629,8 @@ public:
 	// FVulkanEvictable interface.
 	bool CanMove() const override { return false; }
 	bool CanEvict() const override { return false; }
-	void Evict(FVulkanDevice& Device, FVulkanCommandListContext& Context) override; ///evict to system memory
-	void Move(FVulkanDevice& Device, FVulkanCommandListContext& Context, VulkanRHI::FVulkanAllocation& NewAllocation) override; //move to a full new allocation
+	void Evict(FVulkanDevice& Device, const FVulkanContextArray& Contexts) override; ///evict to system memory
+	void Move(FVulkanDevice& Device, const FVulkanContextArray& Contexts, VulkanRHI::FVulkanAllocation& NewAllocation) override; //move to a full new allocation
 	FVulkanTexture* GetEvictableTexture() override { return CanEvict() ? this : nullptr; }
 
 	bool GetTextureResourceInfo(FRHIResourceInfo& OutResourceInfo) const;
@@ -768,7 +769,7 @@ public:
 
 	const FVulkanCpuReadbackBuffer* GetCpuReadbackBuffer() const { return CpuReadbackBuffer; }
 
-	virtual void UpdateLinkedViews() override;
+	virtual void UpdateLinkedViews(const FVulkanContextArray& Contexts) override;
 
 	FVulkanDevice* const Device;
 	VkImage Image = VK_NULL_HANDLE;
@@ -795,6 +796,7 @@ private:
 
 	static const VkImageUsageFlags DefaultImageUsageFlags = 0;
 	void InitViews(VkImageUsageFlags ViewUsageFlags, VkSamplerYcbcrConversion SamplerYcbcrConversion = VK_NULL_HANDLE);
+	void UpdateViews(const FVulkanContextArray& Contexts, VkImageUsageFlags ViewUsageFlags, VkSamplerYcbcrConversion SamplerYcbcrConversion = VK_NULL_HANDLE);
 
 	VkImageTiling Tiling = VK_IMAGE_TILING_MAX_ENUM;
 	VulkanRHI::FVulkanAllocation Allocation;
@@ -979,7 +981,7 @@ public:
 	FVulkanUnorderedAccessView(FRHICommandListBase& RHICmdList, FVulkanDevice& InDevice, FRHIViewableResource* InResource, FRHIViewDesc const& InViewDesc);
 
 	FVulkanViewableResource* GetBaseResource() const;
-	void UpdateView() override;
+	void UpdateView(const FVulkanContextArray& Contexts) override;
 
 	virtual FRHIDescriptorHandle GetBindlessHandle() const override
 	{
@@ -996,7 +998,7 @@ public:
 	FVulkanShaderResourceView(FRHICommandListBase& RHICmdList, FVulkanDevice& InDevice, FRHIViewableResource* InResource, FRHIViewDesc const& InViewDesc);
 
 	FVulkanViewableResource* GetBaseResource() const;
-	void UpdateView() override;
+	void UpdateView(const FVulkanContextArray& Contexts) override;
 
 	virtual FRHIDescriptorHandle GetBindlessHandle() const override
 	{
diff --git a/Engine/Source/Runtime/VulkanRHI/Private/VulkanSRV.cpp b/Engine/Source/Runtime/VulkanRHI/Private/VulkanSRV.cpp
--- a/Engine/Source/Runtime/VulkanRHI/Private/VulkanSRV.cpp
+++ b/Engine/Source/Runtime/VulkanRHI/Private/VulkanSRV.cpp
@@ -1,6 +1,7 @@
 // Copyright Epic Games, Inc. All Rights Reserved.
 
 #include "VulkanRHIPrivate.h"
+#include "VulkanContext.h"
 
 static VkImageViewType GetVkImageViewTypeForDimensionSRV(FRHIViewDesc::EDimension DescDimension, VkImageViewType TextureViewType)
 {
@@ -24,10 +25,10 @@ FVulkanShaderResourceView::FVulkanShaderResourceView(FRHICommandListBase& RHICmd
 	, FVulkanLinkedView(InDevice, GetVkDescriptorTypeForViewDesc(InViewDesc))
 {
 	RHICmdList.EnqueueLambda([this](FRHICommandListBase&)
-		{
-			LinkHead(GetBaseResource()->LinkedViews);
-			UpdateView();
-		});
+	{
+		LinkHead(GetBaseResource()->LinkedViews);
+		UpdateView({});
+	});
 
 	RHICmdList.RHIThreadFence(true);
 }
@@ -42,7 +43,7 @@ FVulkanViewableResource* FVulkanShaderResourceView::GetBaseResource() const
 	return ResourceCast(GetTexture());
 }
 
-void FVulkanShaderResourceView::UpdateView()
+void FVulkanShaderResourceView::UpdateView(const FVulkanContextArray& Contexts)
 {
 #if VULKAN_ENABLE_AGGRESSIVE_STATS
 	SCOPE_CYCLE_COUNTER(STAT_VulkanSRVUpdateTime);
@@ -61,16 +62,16 @@ void FVulkanShaderResourceView::UpdateView()
 			{
 			case FRHIViewDesc::EBufferType::Raw:
 			case FRHIViewDesc::EBufferType::Structured:
-				InitAsStructuredBufferView(Buffer, Info.OffsetInBytes, Info.SizeInBytes);
+				InitAsStructuredBufferView(Contexts, Buffer, Info.OffsetInBytes, Info.SizeInBytes);
 				break;
 
 			case FRHIViewDesc::EBufferType::Typed:
 				check(VKHasAllFlags(Buffer->GetBufferUsageFlags(), VK_BUFFER_USAGE_UNIFORM_TEXEL_BUFFER_BIT));
-				InitAsTypedBufferView(Buffer, Info.Format, Info.OffsetInBytes, Info.SizeInBytes);
+				InitAsTypedBufferView(Contexts, Buffer, Info.Format, Info.OffsetInBytes, Info.SizeInBytes);
 				break;
 
 			case FRHIViewDesc::EBufferType::AccelerationStructure:
-				InitAsAccelerationStructureView(Buffer, Info.OffsetInBytes, Info.SizeInBytes);
+				InitAsAccelerationStructureView(Contexts, Buffer, Info.OffsetInBytes, Info.SizeInBytes);
 				break;
 
 			default:
@@ -107,7 +108,7 @@ void FVulkanShaderResourceView::UpdateView()
 		VulkanViewDesc.NumArraySlices = ArrayNum;
 		VulkanViewDesc.bUseIdentitySwizzle = false;
 		VulkanViewDesc.ImageUsageFlags = ImageUsageFlags;
-		InitAsTextureView(Texture->Image, VulkanViewDesc);
+		InitAsTextureView(Contexts, Texture->Image, VulkanViewDesc);
 	}
 }
 
diff --git a/Engine/Source/Runtime/VulkanRHI/Private/VulkanTexture.cpp b/Engine/Source/Runtime/VulkanRHI/Private/VulkanTexture.cpp
--- a/Engine/Source/Runtime/VulkanRHI/Private/VulkanTexture.cpp
+++ b/Engine/Source/Runtime/VulkanRHI/Private/VulkanTexture.cpp
@@ -970,35 +970,38 @@ void FVulkanDynamicRHI::RHIUpdateTextureReference(FRHICommandListBase& RHICmdLis
 #if PLATFORM_SUPPORTS_BINDLESS_RENDERING
 	if (Device->SupportsBindless() && TextureRef && TextureRef->IsBindless())
 	{
-		RHICmdList.EnqueueLambda(TEXT("FVulkanDynamicRHI::RHIUpdateTextureReference"), [TextureRef, InNewTexture](FRHICommandListBase& ExecutingCmdList)
-		{
-			FVulkanTexture* NewVulkanTexture = ResourceCast(InNewTexture ? InNewTexture : FRHITextureReference::GetDefaultTexture());
-			FVulkanTextureReference* VulkanTextureReference = ResourceCast(TextureRef);
-
-			FVulkanShaderResourceView* VulkanTextureRefSRV = VulkanTextureReference->BindlessView;
-			FRHIDescriptorHandle DestHandle = VulkanTextureRefSRV->GetBindlessHandle();
+		// Workaround for a crash bug where FRHITextureReferences are deleted before this command is executed on the RHI thread.
+		// Take a reference on the FRHITextureReference object to keep it alive.
+		// @todo dev-pr - This should be refactored out when we eventually remove FRHITextureReference.
+		TRefCountPtr<FVulkanTextureReference> Ref = ResourceCast(TextureRef);
 
-			if (DestHandle.IsValid())
+		RHICmdList.EnqueueLambdaMultiPipe(GetEnabledRHIPipelines(), FRHICommandListBase::EThreadFence::Enabled, TEXT("FVulkanDynamicRHI::RHIUpdateTextureReference"),
+			[TextureRef = MoveTemp(Ref), NewTexture = ResourceCast(InNewTexture ? InNewTexture : FRHITextureReference::GetDefaultTexture())](const FVulkanContextArray& Contexts)
 			{
-				checkf(VulkanTextureRefSRV->IsInitialized(), TEXT("TextureReference should always be created with a view of the default texture at least"));
-
-				const FRHITextureDesc& Desc = NewVulkanTexture->GetDesc();
-
-				VulkanTextureRefSRV->Invalidate();
-
-				FVulkanTextureViewDesc VulkanViewDesc;
-				VulkanViewDesc.ViewType = NewVulkanTexture->GetViewType();
-				VulkanViewDesc.AspectFlags = NewVulkanTexture->GetPartialAspectMask();
-				VulkanViewDesc.UEFormat = Desc.Format;
-				VulkanViewDesc.Format = NewVulkanTexture->ViewFormat;
-				VulkanViewDesc.FirstMip = 0u;
-				VulkanViewDesc.NumMips = FMath::Max(Desc.NumMips, (uint8)1u);
-				VulkanViewDesc.ArraySliceIndex = 0u;
-				VulkanViewDesc.NumArraySlices = NewVulkanTexture->GetNumberOfArrayLevels();
-				VulkanViewDesc.bUseIdentitySwizzle = !NewVulkanTexture->SupportsSampling();
-				VulkanTextureRefSRV->InitAsTextureView(NewVulkanTexture->Image, VulkanViewDesc);
-			}
-		});
+				FVulkanShaderResourceView* VulkanTextureRefSRV = TextureRef->BindlessView;
+				FRHIDescriptorHandle DestHandle = VulkanTextureRefSRV->GetBindlessHandle();
+
+				if (DestHandle.IsValid())
+				{
+					checkf(VulkanTextureRefSRV->IsInitialized(), TEXT("TextureReference should always be created with a view of the default texture at least"));
+
+					const FRHITextureDesc& Desc = NewTexture->GetDesc();
+
+					VulkanTextureRefSRV->Invalidate();
+
+					FVulkanTextureViewDesc VulkanViewDesc;
+					VulkanViewDesc.ViewType = NewTexture->GetViewType();
+					VulkanViewDesc.AspectFlags = NewTexture->GetPartialAspectMask();
+					VulkanViewDesc.UEFormat = Desc.Format;
+					VulkanViewDesc.Format = NewTexture->ViewFormat;
+					VulkanViewDesc.FirstMip = 0u;
+					VulkanViewDesc.NumMips = FMath::Max(Desc.NumMips, (uint8)1u);
+					VulkanViewDesc.ArraySliceIndex = 0u;
+					VulkanViewDesc.NumArraySlices = NewTexture->GetNumberOfArrayLevels();
+					VulkanViewDesc.bUseIdentitySwizzle = !NewTexture->SupportsSampling();
+					VulkanTextureRefSRV->UpdateTextureView(Contexts, NewTexture->Image, VulkanViewDesc);
+				}
+			});
 	}
 #endif // PLATFORM_SUPPORTS_BINDLESS_RENDERING
 
@@ -1962,46 +1965,53 @@ FVulkanTexture::FVulkanTexture(FVulkanDevice& InDevice, const FRHITextureCreateD
 #endif // PLATFORM_ANDROID
 
 
-void FVulkanTexture::InitViews(VkImageUsageFlags ViewUsageFlags, VkSamplerYcbcrConversion SamplerYcbcrConversion)
+static FVulkanTextureViewDesc CreateTextureViewDesc(FVulkanTexture* Texture, VkImageUsageFlags ViewUsageFlags, VkSamplerYcbcrConversion SamplerYcbcrConversion)
 {
-	const FRHITextureDesc& Desc = GetDesc();
-	const VkImageViewType ViewType = GetViewType();
-	const VkDescriptorType DescriptorType = SupportsSampling() ? VK_DESCRIPTOR_TYPE_SAMPLED_IMAGE : VK_DESCRIPTOR_TYPE_STORAGE_IMAGE;
-	const bool bIsSRGB = EnumHasAllFlags(Desc.Flags, TexCreate_SRGB);
+	const FRHITextureDesc& Desc = Texture->GetDesc();
+	const VkDescriptorType DescriptorType = Texture->SupportsSampling() ? VK_DESCRIPTOR_TYPE_SAMPLED_IMAGE : VK_DESCRIPTOR_TYPE_STORAGE_IMAGE;
 	const bool bHasUAVFormat = (Desc.UAVFormat != PF_Unknown) && (Desc.UAVFormat != Desc.Format);
-	const bool bRemoveStorageBit = (bIsSRGB || (bHasUAVFormat && !UE::PixelFormat::HasCapabilities(Desc.Format, EPixelFormatCapabilities::UAV)));
-	if (bRemoveStorageBit)
+
+	if (EnumHasAllFlags(Desc.Flags, ETextureCreateFlags::SRGB) || (bHasUAVFormat && !UE::PixelFormat::HasCapabilities(Desc.Format, EPixelFormatCapabilities::UAV)))
 	{
 		ViewUsageFlags = (ViewUsageFlags & ~VK_IMAGE_USAGE_STORAGE_BIT);
 	}
 
-	const bool bUseIdentitySwizzle = 
+	const bool bUseIdentitySwizzle =
 		(DescriptorType != VK_DESCRIPTOR_TYPE_SAMPLED_IMAGE) ||
-		(ViewFormat == VK_FORMAT_UNDEFINED); // External buffer textures also require identity swizzle
+		(Texture->ViewFormat == VK_FORMAT_UNDEFINED); // External buffer textures also require identity swizzle
 
 	FVulkanTextureViewDesc ViewDesc;
-	ViewDesc.ViewType = GetViewType();
-	ViewDesc.AspectFlags = FullAspectMask;
+	ViewDesc.ViewType = Texture->GetViewType();
+	ViewDesc.AspectFlags = Texture->GetFullAspectMask();
 	ViewDesc.UEFormat = Desc.Format;
-	ViewDesc.Format = ViewFormat;
+	ViewDesc.Format = Texture->ViewFormat;
 	ViewDesc.FirstMip = 0;
 	ViewDesc.NumMips = FMath::Max(Desc.NumMips, (uint8)1u);
 	ViewDesc.ArraySliceIndex = 0;
-	ViewDesc.NumArraySlices = GetNumberOfArrayLevels();
+	ViewDesc.NumArraySlices = Texture->GetNumberOfArrayLevels();
 	ViewDesc.bUseIdentitySwizzle = bUseIdentitySwizzle;
 	ViewDesc.ImageUsageFlags = ViewUsageFlags;
 	ViewDesc.SamplerYcbcrConversion = SamplerYcbcrConversion;
 
-	if ((Image != VK_NULL_HANDLE) && (ViewType != VK_IMAGE_VIEW_TYPE_MAX_ENUM))
+	return ViewDesc;
+}
+
+void FVulkanTexture::InitViews(VkImageUsageFlags ViewUsageFlags, VkSamplerYcbcrConversion SamplerYcbcrConversion)
+{
+	if (Image != VK_NULL_HANDLE && GetViewType() != VK_IMAGE_VIEW_TYPE_MAX_ENUM)
 	{
+		const FVulkanTextureViewDesc ViewDesc = CreateTextureViewDesc(this, ViewUsageFlags, SamplerYcbcrConversion);
+		const VkDescriptorType DescriptorType = SupportsSampling() ? VK_DESCRIPTOR_TYPE_SAMPLED_IMAGE : VK_DESCRIPTOR_TYPE_STORAGE_IMAGE;
+
 		if (DefaultView)
 		{
 			DefaultView->Invalidate();
-			DefaultView->InitAsTextureView(Image, ViewDesc);
+			DefaultView->CreateTextureView(Image, ViewDesc);
 		}
 		else
 		{
-			DefaultView = (new FVulkanView(*Device, DescriptorType))->InitAsTextureView(Image, ViewDesc);
+			DefaultView = new FVulkanView(*Device, DescriptorType);
+			DefaultView->CreateTextureView(Image, ViewDesc);
 		}
 
 		if (FullAspectMask == PartialAspectMask)
@@ -2010,18 +2020,65 @@ void FVulkanTexture::InitViews(VkImageUsageFlags ViewUsageFlags, VkSamplerYcbcrC
 		}
 		else
 		{
-			ViewDesc.AspectFlags = PartialAspectMask;
-			ViewDesc.bUseIdentitySwizzle = false;
-			ViewDesc.SamplerYcbcrConversion = VK_NULL_HANDLE;
+			FVulkanTextureViewDesc PartialViewDesc = ViewDesc;
+
+			PartialViewDesc.AspectFlags = PartialAspectMask;
+			PartialViewDesc.bUseIdentitySwizzle = false;
+			PartialViewDesc.SamplerYcbcrConversion = VK_NULL_HANDLE;
 
 			if (PartialView)
 			{
 				PartialView->Invalidate();
-				PartialView->InitAsTextureView(Image, ViewDesc);
+				PartialView->CreateTextureView(Image, PartialViewDesc);
 			}
 			else
 			{
-				PartialView = (new FVulkanView(*Device, DescriptorType))->InitAsTextureView(Image, ViewDesc);
+				PartialView = new FVulkanView(*Device, DescriptorType);
+				PartialView->CreateTextureView(Image, PartialViewDesc);
+			}
+		}
+	}
+}
+
+void FVulkanTexture::UpdateViews(const FVulkanContextArray& Contexts, VkImageUsageFlags ViewUsageFlags, VkSamplerYcbcrConversion SamplerYcbcrConversion)
+{
+	if (Image != VK_NULL_HANDLE && GetViewType() != VK_IMAGE_VIEW_TYPE_MAX_ENUM)
+	{
+		const FVulkanTextureViewDesc ViewDesc = CreateTextureViewDesc(this, ViewUsageFlags, SamplerYcbcrConversion);
+		const VkDescriptorType DescriptorType = SupportsSampling() ? VK_DESCRIPTOR_TYPE_SAMPLED_IMAGE : VK_DESCRIPTOR_TYPE_STORAGE_IMAGE;
+
+		if (DefaultView)
+		{
+			DefaultView->Invalidate();
+			DefaultView->UpdateTextureView(Contexts, Image, ViewDesc);
+		}
+		else
+		{
+			DefaultView = new FVulkanView(*Device, DescriptorType);
+			DefaultView->CreateTextureView(Image, ViewDesc);
+		}
+
+		if (FullAspectMask == PartialAspectMask)
+		{
+			PartialView = DefaultView;
+		}
+		else
+		{
+			FVulkanTextureViewDesc PartialViewDesc = ViewDesc;
+
+			PartialViewDesc.AspectFlags = PartialAspectMask;
+			PartialViewDesc.bUseIdentitySwizzle = false;
+			PartialViewDesc.SamplerYcbcrConversion = VK_NULL_HANDLE;
+
+			if (PartialView)
+			{
+				PartialView->Invalidate();
+				PartialView->UpdateTextureView(Contexts, Image, PartialViewDesc);
+			}
+			else
+			{
+				PartialView = new FVulkanView(*Device, DescriptorType);
+				PartialView->CreateTextureView(Image, PartialViewDesc);
 			}
 		}
 	}
@@ -2042,7 +2099,8 @@ FVulkanView* FVulkanTexture::FindOrAddInternalView(const FVulkanTextureViewDesc&
 	FVulkanView*& VulkanView = InternalViews.FindOrAdd(ViewDesc.GetHash(), nullptr);
 	if (!VulkanView)
 	{
-		VulkanView = (new FVulkanView(*Device, VK_DESCRIPTOR_TYPE_SAMPLED_IMAGE))->InitAsTextureView(Image, ViewDesc);
+		VulkanView = new FVulkanView(*Device, VK_DESCRIPTOR_TYPE_SAMPLED_IMAGE);
+		VulkanView->CreateTextureView(Image, ViewDesc);
 	}
 	return VulkanView;
 }
@@ -2111,9 +2169,9 @@ void FVulkanTexture::AliasTextureResources(FTextureRHIRef& SrcTextureRHI)
 	DefaultLayout = SrcTexture->DefaultLayout;
 }
 
-void FVulkanTexture::UpdateLinkedViews()
+void FVulkanTexture::UpdateLinkedViews(const FVulkanContextArray& Contexts)
 {
-	InitViews(ImageUsageFlags);
+	UpdateViews(Contexts, ImageUsageFlags);
 
 	const bool bRenderTarget = EnumHasAnyFlags(GetDesc().Flags, TexCreate_RenderTargetable | TexCreate_DepthStencilTargetable | TexCreate_ResolveTargetable);
 	if (bRenderTarget)
@@ -2121,16 +2179,16 @@ void FVulkanTexture::UpdateLinkedViews()
 		Device->NotifyDeletedImage(Image, true);
 	}
 
-	for (auto& Pair : InternalViews)
+	for (TPair<uint32, FVulkanView*>& Pair : InternalViews)
 	{
 		delete Pair.Value;
 	}
 	InternalViews.Reset();
 
-	FVulkanViewableResource::UpdateLinkedViews();
+	FVulkanViewableResource::UpdateLinkedViews(Contexts);
 }
 
-void FVulkanTexture::Move(FVulkanDevice& InDevice, FVulkanCommandListContext& Context, VulkanRHI::FVulkanAllocation& NewAllocation)
+void FVulkanTexture::Move(FVulkanDevice& InDevice, const FVulkanContextArray& Contexts, VulkanRHI::FVulkanAllocation& NewAllocation)
 {
 	const uint64 Size = GetMemorySize();
 	static uint64 TotalSize = 0;
@@ -2147,15 +2205,15 @@ void FVulkanTexture::Move(FVulkanDevice& InDevice, FVulkanCommandListContext& Co
 	checkf(bRenderTarget || bUAV, TEXT("Surface must be a RenderTarget or a UAV in order to be moved.  UEFlags=0x%x"), (int32)UEFlags);
 	checkf(Tiling == VK_IMAGE_TILING_OPTIMAL, TEXT("Tiling [%s] is not supported for move, only VK_IMAGE_TILING_OPTIMAL"), VK_TYPE_TO_STRING(VkImageTiling, Tiling));
 
-	InternalMoveSurface(InDevice, Context, NewAllocation);
+	InternalMoveSurface(InDevice, *Contexts[ERHIPipeline::Graphics], NewAllocation);
 	
 	// Swap in the new allocation for this surface
 	Allocation.Swap(NewAllocation);
 
-	UpdateLinkedViews();
+	UpdateLinkedViews(Contexts);
 }
 
-void FVulkanTexture::Evict(FVulkanDevice& InDevice, FVulkanCommandListContext& Context)
+void FVulkanTexture::Evict(FVulkanDevice& InDevice, const FVulkanContextArray& Contexts)
 {
 	check(AliasedTexture == nullptr); //can't evict textures we don't own
 	const uint64 Size = GetMemorySize();
@@ -2190,7 +2248,7 @@ void FVulkanTexture::Evict(FVulkanDevice& InDevice, FVulkanCommandListContext& C
 			checkNoEntry();
 		}
 
-		InternalMoveSurface(InDevice, Context, HostAllocation);
+		InternalMoveSurface(InDevice, *Contexts[ERHIPipeline::Graphics], HostAllocation);
 
 		// Delete the original allocation and swap in the new host allocation
 		Device->GetMemoryManager().FreeVulkanAllocation(Allocation);
@@ -2198,7 +2256,7 @@ void FVulkanTexture::Evict(FVulkanDevice& InDevice, FVulkanCommandListContext& C
 
 		VULKAN_SET_DEBUG_NAME(InDevice, VK_OBJECT_TYPE_IMAGE, Image, TEXT("(FVulkanTexture*)0x%p [hostimage]"), this);
 
-		UpdateLinkedViews();
+		UpdateLinkedViews(Contexts);
 	}
 }
 
diff --git a/Engine/Source/Runtime/VulkanRHI/Private/VulkanUAV.cpp b/Engine/Source/Runtime/VulkanRHI/Private/VulkanUAV.cpp
--- a/Engine/Source/Runtime/VulkanRHI/Private/VulkanUAV.cpp
+++ b/Engine/Source/Runtime/VulkanRHI/Private/VulkanUAV.cpp
@@ -32,7 +32,7 @@ FVulkanUnorderedAccessView::FVulkanUnorderedAccessView(FRHICommandListBase& RHIC
 	RHICmdList.EnqueueLambda([this](FRHICommandListBase&)
 	{
 		LinkHead(GetBaseResource()->LinkedViews);
-		UpdateView();
+		UpdateView({});
 	});
 
 	RHICmdList.RHIThreadFence(true);
@@ -48,7 +48,7 @@ FVulkanViewableResource* FVulkanUnorderedAccessView::GetBaseResource() const
 	return ResourceCast(GetTexture());
 }
 
-void FVulkanUnorderedAccessView::UpdateView()
+void FVulkanUnorderedAccessView::UpdateView(const FVulkanContextArray& Contexts)
 {
 #if VULKAN_ENABLE_AGGRESSIVE_STATS
 	SCOPE_CYCLE_COUNTER(STAT_VulkanUAVUpdateTime);
@@ -69,12 +69,12 @@ void FVulkanUnorderedAccessView::UpdateView()
 			{
 			case FRHIViewDesc::EBufferType::Raw:
 			case FRHIViewDesc::EBufferType::Structured:
-				InitAsStructuredBufferView(Buffer, Info.OffsetInBytes, Info.SizeInBytes);
+				InitAsStructuredBufferView(Contexts, Buffer, Info.OffsetInBytes, Info.SizeInBytes);
 				break;
 
 			case FRHIViewDesc::EBufferType::Typed:
 				check(VKHasAllFlags(Buffer->GetBufferUsageFlags(), VK_BUFFER_USAGE_STORAGE_TEXEL_BUFFER_BIT));
-				InitAsTypedBufferView(Buffer, Info.Format, Info.OffsetInBytes, Info.SizeInBytes);
+				InitAsTypedBufferView(Contexts, Buffer, Info.Format, Info.OffsetInBytes, Info.SizeInBytes);
 				break;
 
 			case FRHIViewDesc::EBufferType::AccelerationStructure:
@@ -111,7 +111,7 @@ void FVulkanUnorderedAccessView::UpdateView()
 		VulkanViewDesc.ArraySliceIndex = ArrayFirst;
 		VulkanViewDesc.NumArraySlices = ArrayNum;
 		VulkanViewDesc.bUseIdentitySwizzle = true;
-		InitAsTextureView(Texture->Image, VulkanViewDesc);
+		InitAsTextureView(Contexts, Texture->Image, VulkanViewDesc);
 	}
 }
 
diff --git a/Engine/Source/Runtime/VulkanRHI/Private/VulkanView.cpp b/Engine/Source/Runtime/VulkanRHI/Private/VulkanView.cpp
--- a/Engine/Source/Runtime/VulkanRHI/Private/VulkanView.cpp
+++ b/Engine/Source/Runtime/VulkanRHI/Private/VulkanView.cpp
@@ -4,6 +4,7 @@
 #include "VulkanBindlessDescriptorManager.h"
 #include "VulkanDevice.h"
 #include "VulkanLLM.h"
+#include "VulkanContext.h"
 
 uint32 FVulkanTextureViewDesc::GetHash() const
 {
@@ -136,7 +137,7 @@ void FVulkanView::Invalidate()
 	Storage.Get<FInvalidatedState>().bInitialized = bIsInitialized;
 }
 
-FVulkanView* FVulkanView::InitAsTypedBufferView(FVulkanBuffer* Buffer, EPixelFormat UEFormat, uint32 InOffset, uint32 InSize)
+void FVulkanView::InitAsTypedBufferView(const FVulkanContextArray& Contexts, FVulkanBuffer* Buffer, EPixelFormat UEFormat, uint32 InOffset, uint32 InSize)
 {
 	// We will need a deferred update if the descriptor was already in use
 	const bool bImmediateUpdate = !IsInitialized();
@@ -184,12 +185,10 @@ FVulkanView* FVulkanView::InitAsTypedBufferView(FVulkanBuffer* Buffer, EPixelFor
 	INC_DWORD_STAT(STAT_VulkanNumBufferViews);
 	// :todo-jn: the buffer view is actually not needed in bindless anymore
 
-	Device.GetBindlessDescriptorManager()->UpdateTexelBuffer(BindlessHandle, DescriptorType, ViewInfo, bImmediateUpdate);
-
-	return this;
+	Device.GetBindlessDescriptorManager()->UpdateTexelBuffer(Contexts, BindlessHandle, DescriptorType, ViewInfo, bImmediateUpdate);
 }
 
-FVulkanView* FVulkanView::InitAsTextureView(VkImage InImage, const FVulkanTextureViewDesc& ViewDesc)
+void FVulkanView::InitAsTextureView(const FVulkanContextArray& Contexts, VkImage InImage, const FVulkanTextureViewDesc& ViewDesc)
 {
 	// We will need a deferred update if the descriptor was already in use
 	const bool bImmediateUpdate = !IsInitialized();
@@ -279,12 +278,20 @@ FVulkanView* FVulkanView::InitAsTextureView(VkImage InImage, const FVulkanTextur
 	}
 
 	const bool bDepthOrStencilAspect = (ViewDesc.AspectFlags & (VK_IMAGE_ASPECT_DEPTH_BIT | VK_IMAGE_ASPECT_STENCIL_BIT)) != 0;
-	Device.GetBindlessDescriptorManager()->UpdateImage(BindlessHandle, DescriptorType, TV.View, bDepthOrStencilAspect, bImmediateUpdate);
+	Device.GetBindlessDescriptorManager()->UpdateImage(Contexts, BindlessHandle, DescriptorType, TV.View, bDepthOrStencilAspect, bImmediateUpdate);
+}
+
+void FVulkanView::CreateTextureView(VkImage InImage, const FVulkanTextureViewDesc& ViewDesc)
+{
+	InitAsTextureView({}, InImage, ViewDesc);
+}
 
-	return this;
+void FVulkanView::UpdateTextureView(const FVulkanContextArray& Contexts, VkImage InImage, const FVulkanTextureViewDesc& ViewDesc)
+{
+	InitAsTextureView(Contexts, InImage, ViewDesc);
 }
 
-FVulkanView* FVulkanView::InitAsStructuredBufferView(FVulkanBuffer* Buffer, uint32 InOffset, uint32 InSize)
+void FVulkanView::InitAsStructuredBufferView(const FVulkanContextArray& Contexts, FVulkanBuffer* Buffer, uint32 InOffset, uint32 InSize)
 {
 	// We will need a deferred update if the descriptor was already in use
 	const bool bImmediateUpdate = !IsInitialized();
@@ -307,12 +314,10 @@ FVulkanView* FVulkanView::InitAsStructuredBufferView(FVulkanBuffer* Buffer, uint
 
 	SBV.Size = InSize;
 
-	Device.GetBindlessDescriptorManager()->UpdateBuffer(BindlessHandle, DescriptorType, Buffer->GetHandle(), TotalOffset, InSize, bImmediateUpdate);
-
-	return this;
+	Device.GetBindlessDescriptorManager()->UpdateBuffer(Contexts, BindlessHandle, DescriptorType, Buffer->GetHandle(), TotalOffset, InSize, bImmediateUpdate);
 }
 
-FVulkanView* FVulkanView::InitAsAccelerationStructureView(FVulkanBuffer* Buffer, uint32 Offset, uint32 Size)
+void FVulkanView::InitAsAccelerationStructureView(const FVulkanContextArray& Contexts, FVulkanBuffer* Buffer, uint32 Offset, uint32 Size)
 {
 	check(GetViewType() == EType::Null);
 	Storage.Emplace<FAccelerationStructureView>();
@@ -327,18 +332,16 @@ FVulkanView* FVulkanView::InitAsAccelerationStructureView(FVulkanBuffer* Buffer,
 
 	VERIFYVULKANRESULT(VulkanDynamicAPI::vkCreateAccelerationStructureKHR(Device.GetHandle(), &CreateInfo, VULKAN_CPU_ALLOCATOR, &ASV.Handle));
 
-	Device.GetBindlessDescriptorManager()->UpdateAccelerationStructure(BindlessHandle, ASV.Handle);
-
-	return this;
+	Device.GetBindlessDescriptorManager()->UpdateAccelerationStructure(Contexts, BindlessHandle, ASV.Handle, true);
 }
 
 //
 // FVulkanViewableResource
 
-void FVulkanViewableResource::UpdateLinkedViews()
+void FVulkanViewableResource::UpdateLinkedViews(const FVulkanContextArray& Contexts)
 {
 	for (FVulkanLinkedView* View = LinkedViews; View; View = View->Next())
 	{
-		View->UpdateView();
+		View->UpdateView(Contexts);
 	}
 }
diff --git a/Engine/Source/Runtime/VulkanRHI/Private/VulkanViewport.cpp b/Engine/Source/Runtime/VulkanRHI/Private/VulkanViewport.cpp
--- a/Engine/Source/Runtime/VulkanRHI/Private/VulkanViewport.cpp
+++ b/Engine/Source/Runtime/VulkanRHI/Private/VulkanViewport.cpp
@@ -381,7 +381,11 @@ void FVulkanViewport::InitImages(FVulkanContextCommon& Context, TConstArrayView<
 		ViewDesc.NumMips = 1;
 		ViewDesc.ArraySliceIndex = 0;
 		ViewDesc.NumArraySlices = 1;
-		TextureViews.Add((new FVulkanView(Device, VK_DESCRIPTOR_TYPE_SAMPLED_IMAGE))->InitAsTextureView(Images[Index], ViewDesc));
+
+		FVulkanView* NewView = new FVulkanView(Device, VK_DESCRIPTOR_TYPE_SAMPLED_IMAGE);
+		NewView->CreateTextureView(Images[Index], ViewDesc);
+
+		TextureViews.Add(NewView);
 
 		// Clear the swapchain to avoid a validation warning, and transition to PresentSrc
 		{
