VulkanRHI: Android Swappy GPUTime fixes

#tests Internal QA replay run

diff --git a/Engine/Source/Runtime/VulkanRHI/Private/Android/VulkanAndroidPlatform.cpp b/Engine/Source/Runtime/VulkanRHI/Private/Android/VulkanAndroidPlatform.cpp
--- a/Engine/Source/Runtime/VulkanRHI/Private/Android/VulkanAndroidPlatform.cpp
+++ b/Engine/Source/Runtime/VulkanRHI/Private/Android/VulkanAndroidPlatform.cpp
@@ -34,14 +34,21 @@ namespace AndroidVulkan
 {
 	void VKSwappyPostWaitCallback(void*, int64_t cpu_time_ns, int64_t gpu_time_ns)
 	{
-		GRHIGPUFrameTimeHistory.PushFrameCycles(1'000'000'000.0, gpu_time_ns);
+		const double GPUTimeInSeconds = (double)gpu_time_ns / 1000000000.0;
+		uint64 Cycles = GPUTimeInSeconds / FPlatformTime::GetSecondsPerCycle64();
+		GetDynamicRHI<FVulkanDynamicRHI>()->ExternalGPUTime = Cycles;
 	}
 
+	static bool GSwappyRegistered = false;
 	void SetSwappyPostWaitCallback()
 	{
-		SwappyTracer Tracer = { 0 };
-		Tracer.postWait = VKSwappyPostWaitCallback;
-		SwappyVk_injectTracer(&Tracer);
+		if (!GSwappyRegistered)
+		{
+			SwappyTracer Tracer = { 0 };
+			Tracer.postWait = VKSwappyPostWaitCallback;
+			SwappyVk_injectTracer(&Tracer);
+			GSwappyRegistered = true;
+		}
 
 		int32 FrameTimeFenceInMillis = FAndroidPlatformRHIFramePacer::CVarSwappyGPUFrameTimeFence.GetValueOnAnyThread();
 
diff --git a/Engine/Source/Runtime/VulkanRHI/Private/VulkanDynamicRHI.h b/Engine/Source/Runtime/VulkanRHI/Private/VulkanDynamicRHI.h
--- a/Engine/Source/Runtime/VulkanRHI/Private/VulkanDynamicRHI.h
+++ b/Engine/Source/Runtime/VulkanRHI/Private/VulkanDynamicRHI.h
@@ -334,6 +334,8 @@ public:
 	void EnqueueEndOfPipeTask(TUniqueFunction<void()> TaskFunc, TUniqueFunction<void(FVulkanPayload&)> ModifyPayloadCallback = {});
 	void CompletePayload(FVulkanPayload* Payload);
 
+	TOptional<uint64> ExternalGPUTime;
+
 	// Returns true if the calling thread is the interrupt thread
 	// (or acting as the interrupt thread when USE_INTERRUPT_THREAD is zero).
 	bool IsInInterruptThread() const;
diff --git a/Engine/Source/Runtime/VulkanRHI/Private/VulkanQueue.cpp b/Engine/Source/Runtime/VulkanRHI/Private/VulkanQueue.cpp
--- a/Engine/Source/Runtime/VulkanRHI/Private/VulkanQueue.cpp
+++ b/Engine/Source/Runtime/VulkanRHI/Private/VulkanQueue.cpp
@@ -216,6 +216,12 @@ void FVulkanQueue::SubmitPayloads(TArrayView<FVulkanPayload*> Payloads, TMap<VkS
 		Payload->PreExecute();
 
 #if RHI_NEW_GPU_PROFILER
+		const bool bHasExternalGPUTime = Payload->ExternalGPUTime.IsSet();
+		if (bHasExternalGPUTime)
+		{
+			Payload->EventStream.Emplace<UE::RHI::GPUProfiler::FEvent::FFrameTime>(*Payload->ExternalGPUTime);
+		}
+
 		if (Payload->EndFrameEvent.IsSet())
 		{
 			Payload->EndFrameEvent->CPUTimestamp = SubmitTime;
diff --git a/Engine/Source/Runtime/VulkanRHI/Private/VulkanRHI.cpp b/Engine/Source/Runtime/VulkanRHI/Private/VulkanRHI.cpp
--- a/Engine/Source/Runtime/VulkanRHI/Private/VulkanRHI.cpp
+++ b/Engine/Source/Runtime/VulkanRHI/Private/VulkanRHI.cpp
@@ -1189,35 +1189,41 @@ void FVulkanDynamicRHI::RHIEndFrame(const FRHIEndFrameArgs& Args)
 	};
 
 	EnqueueEndOfPipeTask(MoveTemp(Lambda), [&](FVulkanPayload& Payload)
+	{
+		// Modify the payloads the EOP task will submit to include
+		// a new timing struct and a frame boundary event.
+
+		Payload.Timing = CurrentTimingPerQueue.CreateNew(Payload.Queue);
+
+		ERHIPipeline Pipeline;
+		switch (Payload.Queue.QueueType)
 		{
-			// Modify the payloads the EOP task will submit to include
-			// a new timing struct and a frame boundary event.
+		default: checkNoEntry(); [[fallthrough]];
+		case EVulkanQueueType::Graphics:     Pipeline = ERHIPipeline::Graphics; break;
+		case EVulkanQueueType::AsyncCompute: Pipeline = ERHIPipeline::AsyncCompute; break;
 
-			Payload.Timing = CurrentTimingPerQueue.CreateNew(Payload.Queue);
+		case EVulkanQueueType::Transfer:
+			// There is currently no high level RHI copy queue support
+			Pipeline = ERHIPipeline::None;
+			break;
+		}
 
-			ERHIPipeline Pipeline;
-			switch (Payload.Queue.QueueType)
-			{
-			default: checkNoEntry(); [[fallthrough]];
-			case EVulkanQueueType::Graphics:     Pipeline = ERHIPipeline::Graphics; break;
-			case EVulkanQueueType::AsyncCompute: Pipeline = ERHIPipeline::AsyncCompute; break;
-
-			case EVulkanQueueType::Transfer:
-				// There is currently no high level RHI copy queue support
-				Pipeline = ERHIPipeline::None;
-				break;
-			}
+		if(ExternalGPUTime.IsSet())
+		{
+			Payload.ExternalGPUTime = ExternalGPUTime.GetValue();
+			ExternalGPUTime.Reset();
+		}
 
-			// CPU timestamp for the frame boundary event is filled in by the submission thread
-			Payload.EndFrameEvent = UE::RHI::GPUProfiler::FEvent::FFrameBoundary(0, Args.FrameNumber
+		// CPU timestamp for the frame boundary event is filled in by the submission thread
+		Payload.EndFrameEvent = UE::RHI::GPUProfiler::FEvent::FFrameBoundary(0, Args.FrameNumber
 #if WITH_RHI_BREADCRUMBS
-				, (Pipeline != ERHIPipeline::None) ? Args.GPUBreadcrumbs[Pipeline] : nullptr
+			, (Pipeline != ERHIPipeline::None) ? Args.GPUBreadcrumbs[Pipeline] : nullptr
 #endif
 #if STATS
-				, Args.StatsFrame
+			, Args.StatsFrame
 #endif
-			);
-		});
+		);
+	});
 #else
 	{
 		FVulkanPlatformCommandList* Payloads = new FVulkanPlatformCommandList;
diff --git a/Engine/Source/Runtime/VulkanRHI/Private/VulkanSubmission.h b/Engine/Source/Runtime/VulkanRHI/Private/VulkanSubmission.h
--- a/Engine/Source/Runtime/VulkanRHI/Private/VulkanSubmission.h
+++ b/Engine/Source/Runtime/VulkanRHI/Private/VulkanSubmission.h
@@ -83,6 +83,7 @@ protected:
 
 #if RHI_NEW_GPU_PROFILER
 	TOptional<FVulkanTiming*> Timing;
+	TOptional<uint64> ExternalGPUTime;
 	TOptional<UE::RHI::GPUProfiler::FEvent::FFrameBoundary> EndFrameEvent;
 	UE::RHI::GPUProfiler::FEventStream EventStream;
 #else
