VulkanRHI: Replace global bool with UE_CALL_ONCE.

#rnx

diff --git a/Engine/Source/Runtime/VulkanRHI/Private/Android/VulkanAndroidPlatform.cpp b/Engine/Source/Runtime/VulkanRHI/Private/Android/VulkanAndroidPlatform.cpp
--- a/Engine/Source/Runtime/VulkanRHI/Private/Android/VulkanAndroidPlatform.cpp
+++ b/Engine/Source/Runtime/VulkanRHI/Private/Android/VulkanAndroidPlatform.cpp
@@ -39,16 +39,14 @@ namespace AndroidVulkan
 		GetDynamicRHI<FVulkanDynamicRHI>()->ExternalGPUTime = Cycles;
 	}
 
-	static bool GSwappyRegistered = false;
 	void SetSwappyPostWaitCallback()
 	{
-		if (!GSwappyRegistered)
+		UE_CALL_ONCE([]
 		{
 			SwappyTracer Tracer = { 0 };
 			Tracer.postWait = VKSwappyPostWaitCallback;
 			SwappyVk_injectTracer(&Tracer);
-			GSwappyRegistered = true;
-		}
+		});
 
 		int32 FrameTimeFenceInMillis = FAndroidPlatformRHIFramePacer::CVarSwappyGPUFrameTimeFence.GetValueOnAnyThread();
 
