VulkanRHI: Add support for reserved resources (sparse bindings for buffers and textures)

#rnx

diff --git a/Engine/Config/Android/AndroidEngine.ini b/Engine/Config/Android/AndroidEngine.ini
--- a/Engine/Config/Android/AndroidEngine.ini
+++ b/Engine/Config/Android/AndroidEngine.ini
@@ -97,6 +97,8 @@ r.Vulkan.AllowSplitBarriers=0
 r.Vulkan.DiagnosticBuffer=1
 ; Transient memory allocator
 r.Vulkan.EnableTransientResourceAllocator=0
+; Disable sparse binding and residency
+r.Vulkan.AllowSparseResidency=0
 
 ; enable stripping debug strings from streamable manager in shipping to save memory
 s.StreamableStripDebugName=true
diff --git a/Engine/Source/Runtime/VulkanRHI/Private/VulkanBarriers.cpp b/Engine/Source/Runtime/VulkanRHI/Private/VulkanBarriers.cpp
--- a/Engine/Source/Runtime/VulkanRHI/Private/VulkanBarriers.cpp
+++ b/Engine/Source/Runtime/VulkanRHI/Private/VulkanBarriers.cpp
@@ -1627,6 +1627,9 @@ void FVulkanCommandListContext::RHIEndTransitions(TArrayView<const FRHITransitio
 		// Wait on semaphores before the acquire barriers
 		ProcessCrossQueueSemaphores(*this, Transitions, bIsBeginTransition);
 
+		// Update reserved resource memory mapping
+		HandleReservedResourceCommits(Transitions);
+
 		ProcessTransitionSync2(*this, Transitions, bIsBeginTransition);
 	}
 	else
diff --git a/Engine/Source/Runtime/VulkanRHI/Private/VulkanBuffer.cpp b/Engine/Source/Runtime/VulkanRHI/Private/VulkanBuffer.cpp
--- a/Engine/Source/Runtime/VulkanRHI/Private/VulkanBuffer.cpp
+++ b/Engine/Source/Runtime/VulkanRHI/Private/VulkanBuffer.cpp
@@ -37,6 +37,14 @@ static FAutoConsoleVariableRef CVarVulkanForceStagingBufferOnLock(
 	ECVF_RenderThreadSafe
 );
 
+static TAutoConsoleVariable<int32> CVarVulkanSparseBufferAllocSizeMB(
+	TEXT("r.Vulkan.SparseBufferAllocSizeMB"),
+	16,
+	TEXT("Size of the backing memory blocks for sparse buffers in megabytes (default 16MB)."),
+	ECVF_ReadOnly
+);
+
+
 static FORCEINLINE FVulkanPendingBufferLock GetPendingBufferLock(FVulkanBuffer* Buffer)
 {
 	FVulkanPendingBufferLock PendingLock;
@@ -129,15 +137,50 @@ FVulkanBuffer::FVulkanBuffer(FVulkanDevice& InDevice, const FRHIBufferCreateDesc
 	{
 		const bool bUnifiedMem = InDevice.HasUnifiedMemory();
 		const uint32 BufferAlignment = VulkanRHI::FMemoryManager::CalculateBufferAlignment(InDevice, CreateDesc.Usage, bZeroSize);
+		const bool bIsSparseResource = EnumHasAnyFlags(CreateDesc.Usage, BUF_ReservedResource);
 
 		if (InTransientHeapAllocation != nullptr)
 		{
+			checkf(!bIsSparseResource, TEXT("Resources with the ReservedResource flag cannot be backed by a transient allocation."));
 			CurrentBufferAlloc.Alloc = FVulkanTransientHeap::GetVulkanAllocation(*InTransientHeapAllocation);
 			CurrentBufferAlloc.HostPtr = bUnifiedMem ? CurrentBufferAlloc.Alloc.GetMappedPointer(&Device) : nullptr;
 			CurrentBufferAlloc.DeviceAddress = GetBufferDeviceAddress(InDevice, CurrentBufferAlloc.Alloc.GetBufferHandle()) + CurrentBufferAlloc.Alloc.Offset;
 			check(CurrentBufferAlloc.Alloc.Offset % BufferAlignment == 0);
 			check(CurrentBufferAlloc.Alloc.Size >= CreateDesc.Size);
 		}
+		else if (bIsSparseResource)
+		{
+			checkf(!EnumHasAnyFlags(GetUsage(), BUF_Dynamic) && !EnumHasAnyFlags(GetUsage(), BUF_Volatile), TEXT("The ReservedResource flag is incompatible with dynamic and volatile buffers."));
+
+			const uint64 SparseBlockSize = (uint64)CVarVulkanSparseBufferAllocSizeMB.GetValueOnAnyThread() * 1024 * 1024;
+			const uint32 BufferSize = Align(GetSize(), GRHIGlobals.ReservedResources.TileSizeInBytes);
+			const VkBufferCreateFlags BufferCreateFlags = VK_BUFFER_CREATE_SPARSE_BINDING_BIT | VK_BUFFER_CREATE_SPARSE_RESIDENCY_BIT;
+
+			CurrentBufferAlloc.Alloc.Size = BufferSize;
+			const VkBuffer BufferHandle = Device.CreateBuffer(BufferSize, BufferUsageFlags, BufferCreateFlags);
+			CurrentBufferAlloc.Alloc.VulkanHandle = (uint64)BufferHandle;
+			CurrentBufferAlloc.DeviceAddress = GetBufferDeviceAddress(Device, BufferHandle);
+
+			ReservedResourceData = MakeUnique<FVulkanReservedResourceData>();
+			ReservedResourceData->LastBlockSize = Align(GetSize() % SparseBlockSize, GRHIGlobals.ReservedResources.TileSizeInBytes);
+			ReservedResourceData->BlockCount = GetSize() / SparseBlockSize;
+			if (ReservedResourceData->LastBlockSize == 0)
+			{
+				ReservedResourceData->LastBlockSize = SparseBlockSize;
+			}
+			else
+			{
+				ReservedResourceData->BlockCount++;
+			}
+			VulkanRHI::vkGetBufferMemoryRequirements(Device.GetHandle(), BufferHandle, &ReservedResourceData->MemoryRequirements);
+			checkf(SparseBlockSize % ReservedResourceData->MemoryRequirements.alignment == 0, 
+				TEXT("The value of r.Vulkan.SparseBufferAllocSizeMB (%dMB) is not compatible with this resource's tile size requirements (%d)."),
+				CVarVulkanSparseBufferAllocSizeMB.GetValueOnAnyThread(), ReservedResourceData->MemoryRequirements.alignment);
+			checkf(ReservedResourceData->LastBlockSize % ReservedResourceData->MemoryRequirements.alignment == 0,
+				TEXT("The value of GRHIGlobals.ReservedResources.TileSizeInBytes (%d) is not compatible with this resource's tile size requirements (%d)."),
+				GRHIGlobals.ReservedResources.TileSizeInBytes, ReservedResourceData->MemoryRequirements.alignment);
+			UpdateVulkanBufferStats(GetDesc(), BufferSize, true);
+		}
 		else
 		{
 			AllocateMemory(CurrentBufferAlloc);
@@ -147,6 +190,12 @@ FVulkanBuffer::FVulkanBuffer(FVulkanDevice& InDevice, const FRHIBufferCreateDesc
 
 FVulkanBuffer::~FVulkanBuffer()
 {
+	if (ReservedResourceData)
+	{
+		Device.GetDeferredDeletionQueue().EnqueueResource(VulkanRHI::FDeferredDeletionQueue2::EType::Buffer, GetHandle());
+		CommitReservedResource(0);
+	}
+
 	VULKAN_TRACK_OBJECT_DELETE(FVulkanBuffer, this);
 	ReleaseOwnership();
 }
@@ -183,6 +232,87 @@ void FVulkanBuffer::AllocateMemory(FBufferAlloc& OutAlloc)
 	UpdateVulkanBufferStats(GetDesc(), NewBufferAlloc.Alloc.Size, true);
 }
 
+TArray<VkSparseMemoryBind> FVulkanBuffer::CommitReservedResource(uint64 RequiredCommitSizeInBytes)
+{
+	checkf(RequiredCommitSizeInBytes <= GetSize(), TEXT("Commit size too large for reserve resource's maximum size."));
+	checkf(ReservedResourceData, TEXT("CommitReservedResource called on a resource without any ReservedResourceData.  Resource needs to be created with the ReservedResource flag."));
+
+	// Figure out the allocation needs
+	const uint64 SparseBlockSize = (uint64)CVarVulkanSparseBufferAllocSizeMB.GetValueOnAnyThread() * 1024 * 1024;
+	const uint64 RequiredBlockCount = Align(RequiredCommitSizeInBytes, SparseBlockSize) / SparseBlockSize;
+	const uint64 CurrentBlockCount = ReservedResourceData->Allocations.Num();
+
+	// Figure out the tile needs
+	const uint64 TileSize = GRHIGlobals.ReservedResources.TileSizeInBytes;
+	const uint64 RequiredTileCount = Align(RequiredCommitSizeInBytes, TileSize) / TileSize;
+	const uint64 CurrentTileCount = ReservedResourceData->NumCommittedTiles;
+
+	TArray<VkSparseMemoryBind> Binds;
+	if (RequiredTileCount > CurrentTileCount)
+	{
+		// Make sure we will have all the memory we need
+		ReservedResourceData->Allocations.Reserve(RequiredBlockCount);
+		for (uint64 BlockCount = CurrentBlockCount; BlockCount < RequiredBlockCount; BlockCount++)
+		{
+			const uint32 BlockSize = (BlockCount == ReservedResourceData->BlockCount - 1) ? ReservedResourceData->LastBlockSize : SparseBlockSize;
+			FVulkanReservedResourceData::FSparseAllocation& NewBlockAlloc = ReservedResourceData->Allocations.AddDefaulted_GetRef();
+			if (!Device.GetMemoryManager().AllocateBufferPooled(
+				NewBlockAlloc.Allocation, nullptr, BlockSize,
+				ReservedResourceData->MemoryRequirements.alignment, BufferUsageFlags,
+				VK_MEMORY_PROPERTY_DEVICE_LOCAL_BIT, VulkanRHI::EVulkanAllocationMetaMultiBuffer,
+				__FILE__, __LINE__))
+			{
+				Device.GetMemoryManager().HandleOOM();
+			}
+			NewBlockAlloc.ResourceOffset = BlockCount * SparseBlockSize;
+		}
+
+		// Bind all the individual tiles
+		Binds.Reserve(RequiredTileCount - CurrentTileCount);
+		for (uint64 TileIndex = CurrentTileCount; TileIndex < RequiredTileCount; TileIndex++)
+		{
+			const uint32 TileOffset = (TileIndex * TileSize);
+
+			const uint32 AllocIndex = TileOffset / SparseBlockSize;
+			const FVulkanReservedResourceData::FSparseAllocation& SparseAlloc = ReservedResourceData->Allocations[AllocIndex];
+
+			VkSparseMemoryBind& NewBind = Binds.AddZeroed_GetRef();
+			NewBind.size = TileSize;
+			NewBind.memory = SparseAlloc.Allocation.GetDeviceMemoryHandle(&Device);
+			NewBind.memoryOffset = SparseAlloc.Allocation.Offset + TileOffset - SparseAlloc.ResourceOffset;
+			NewBind.resourceOffset = TileOffset;
+		}
+
+		const int64 CommitDeltaInBytes = TileSize * Binds.Num();
+		UE::RHICore::UpdateReservedResourceStatsOnCommit(CommitDeltaInBytes, true /* bBuffer */, true /* bCommitting */);
+	}
+	else if (RequiredTileCount < CurrentTileCount)
+	{
+		Binds.Reserve(CurrentTileCount - RequiredTileCount);
+		for (int64 TileIndex = CurrentTileCount - 1; TileIndex >= (int64)RequiredTileCount; TileIndex--)
+		{
+			VkSparseMemoryBind& NewBind = Binds.AddZeroed_GetRef();
+			NewBind.size = TileSize;
+			NewBind.memory = VK_NULL_HANDLE;
+			NewBind.memoryOffset = 0;
+			NewBind.resourceOffset = (TileIndex * TileSize);
+		}
+
+		// Free up unused blocks (deletion is deferred, don't have to wait for binds to be submitted)
+		while (ReservedResourceData->Allocations.Num() > RequiredBlockCount)
+		{
+			ReservedResourceData->Allocations.Last().Allocation.Free(Device);
+			ReservedResourceData->Allocations.Pop();
+		}
+
+		const int64 CommitDeltaInBytes = TileSize * Binds.Num();
+		UE::RHICore::UpdateReservedResourceStatsOnCommit(CommitDeltaInBytes, true /* bBuffer */, false /* bCommitting */);
+	}
+
+	ReservedResourceData->NumCommittedTiles = RequiredTileCount;
+	return Binds;
+}
+
 void* FVulkanBuffer::Lock(FRHICommandListBase& RHICmdList, EResourceLockMode LockMode, uint32 LockSize, uint32 Offset)
 {
 	void* Data = nullptr;
@@ -203,6 +333,7 @@ void* FVulkanBuffer::Lock(FRHICommandListBase& RHICmdList, EResourceLockMode Loc
 	const bool bStatic = EnumHasAnyFlags(GetUsage(), BUF_Static) || !bDynamic;
 	const bool bUAV = EnumHasAnyFlags(GetUsage(), BUF_UnorderedAccess);
 	const bool bSR = EnumHasAnyFlags(GetUsage(), BUF_ShaderResource);
+	const bool bIsSparseResource = EnumHasAnyFlags(GetUsage(), BUF_ReservedResource);
 
 	check(bStatic || bDynamic || bUAV || bSR);
 
@@ -278,7 +409,7 @@ void* FVulkanBuffer::Lock(FRHICommandListBase& RHICmdList, EResourceLockMode Loc
 			DataOffset = Offset;
 			LockStatus = ELockStatus::PersistentMapping;
 		}
-		else if (bStatic || GVulkanForceStagingBufferOnLock)
+		else if (bStatic || GVulkanForceStagingBufferOnLock || bIsSparseResource)
 		{
 			FVulkanPendingBufferLock PendingLock;
 			PendingLock.Offset = Offset;
diff --git a/Engine/Source/Runtime/VulkanRHI/Private/VulkanContext.cpp b/Engine/Source/Runtime/VulkanRHI/Private/VulkanContext.cpp
--- a/Engine/Source/Runtime/VulkanRHI/Private/VulkanContext.cpp
+++ b/Engine/Source/Runtime/VulkanRHI/Private/VulkanContext.cpp
@@ -208,6 +208,44 @@ void FVulkanContextCommon::PrepareNewCommandBuffer(FVulkanPayload& Payload)
 	NewCommandBuffer->Begin(OptionalTimestampQueryPool, RenderPassHandle);
 }
 
+void FVulkanContextCommon::HandleReservedResourceCommits(TArrayView<const FRHITransition*> Transitions)
+{
+	TArray<FVulkanCommitReservedResourceDesc> ReservedResourcesToCommit;
+	for (const FRHITransition* Transition : Transitions)
+	{
+		const FVulkanTransitionData* TransitionData = Transition->GetPrivateData<FVulkanTransitionData>();
+		for (const FRHITransitionInfo& Info : TransitionData->TransitionInfos)
+		{
+			if (const FRHICommitResourceInfo* CommitInfo = Info.CommitInfo.GetPtrOrNull())
+			{
+				if (Info.Type == FRHITransitionInfo::EType::Buffer)
+				{
+					FVulkanCommitReservedResourceDesc CommitDesc;
+					CommitDesc.Resource = Info.Buffer;
+					CommitDesc.CommitSizeInBytes = CommitInfo->SizeInBytes;
+
+					checkf(CommitDesc.Resource, TEXT("FVulkanCommitReservedResourceDesc::Resource must be set"));
+
+					ReservedResourcesToCommit.Add(CommitDesc);
+				}
+				else
+				{
+					checkNoEntry();
+				}
+			}
+		}
+
+	}
+
+	if (!ReservedResourcesToCommit.IsEmpty())
+	{
+		GetPayload(EPhase::UpdateReservedResources).ReservedResourcesToCommit.Append(ReservedResourcesToCommit);
+	}
+}
+
+
+
+
 
 FVulkanCommandListContext::FVulkanCommandListContext(FVulkanDevice& InDevice, ERHIPipeline InPipeline, FVulkanCommandListContext* InImmediate)
 	: FVulkanContextCommon(InDevice, *InDevice.GetQueue(InPipeline), EVulkanCommandBufferType::Primary)
@@ -673,6 +711,9 @@ bool FVulkanCommandListContext::AcquirePoolSetAndDescriptorsIfNeeded(const FVulk
 	return false;
 }
 
+
+
+
 TLockFreePointerListUnordered<FVulkanUploadContext, PLATFORM_CACHE_LINE_SIZE> FVulkanUploadContext::Pool;
 void FVulkanUploadContext::DestroyPool()
 {
diff --git a/Engine/Source/Runtime/VulkanRHI/Private/VulkanContext.h b/Engine/Source/Runtime/VulkanRHI/Private/VulkanContext.h
--- a/Engine/Source/Runtime/VulkanRHI/Private/VulkanContext.h
+++ b/Engine/Source/Runtime/VulkanRHI/Private/VulkanContext.h
@@ -60,6 +60,7 @@ protected:
 	enum class EPhase
 	{
 		Wait,
+		UpdateReservedResources,
 		Execute,
 		Signal
 	};
@@ -207,6 +208,8 @@ public:
 		return nullptr;
 	}
 
+	void HandleReservedResourceCommits(TArrayView<const FRHITransition*> Transitions);
+
 protected:
 
 	// End the current payload, insert the parallel rendering work, and leave context with a valid payload
@@ -411,11 +414,9 @@ public:
 
 	bool AcquirePoolSetAndDescriptorsIfNeeded(const class FVulkanDescriptorSetsLayout& Layout, bool bNeedDescriptors, VkDescriptorSet* OutDescriptors);
 
-
 protected:
 	void ExtractDepthStencilLayouts(const FRHIRenderPassInfo& InInfo, VkImageLayout& OutDepthLayout, VkImageLayout& OutStencilLayout);
 	void FillDynamicRenderingInfo(const FRHIRenderPassInfo& RenderPassInfo, FVulkanDynamicRenderingInfo& OutRenderingInfo);
-	static FVulkanPayload* MergePayloads(TArray<FVulkanPayload*>&& Payloads);
 
 	FVulkanCommandListContext* Immediate;
 	const ERHIPipeline RHIPipeline;
diff --git a/Engine/Source/Runtime/VulkanRHI/Private/VulkanGenericPlatform.cpp b/Engine/Source/Runtime/VulkanRHI/Private/VulkanGenericPlatform.cpp
--- a/Engine/Source/Runtime/VulkanRHI/Private/VulkanGenericPlatform.cpp
+++ b/Engine/Source/Runtime/VulkanRHI/Private/VulkanGenericPlatform.cpp
@@ -15,6 +15,16 @@ static TAutoConsoleVariable<int32> CVarVulkanUseProfileCheck(
 	ECVF_ReadOnly
 );
 
+static TAutoConsoleVariable<int32> CVarVulkanAllowSparseResidency(
+	TEXT("r.Vulkan.AllowSparseResidency"),
+	1,
+	TEXT("Toggles the use of sparse residency for buffers and images:\n")
+	TEXT("0 to disable\n")
+	TEXT("1 to allow when supported (default)\n"),
+	ECVF_ReadOnly
+);
+
+
 void FVulkanGenericPlatform::SetupFeatureLevels(TArrayView<EShaderPlatform> ShaderPlatformForFeatureLevel)
 {
 	checkf(ERHIFeatureLevel::Num == ShaderPlatformForFeatureLevel.Num(), TEXT("ShaderPlatformForFeatureLevel is not the right size to fit all Feature Level values."));
@@ -105,17 +115,20 @@ TArray<FString> FVulkanGenericPlatform::GetPSOCacheFilenames()
 
 void FVulkanGenericPlatform::RestrictEnabledPhysicalDeviceFeatures(FVulkanPhysicalDeviceFeatures* InOutFeaturesToEnable)
 {
-	// Disable everything sparse-related
-	InOutFeaturesToEnable->Core_1_0.shaderResourceResidency = VK_FALSE;
-	InOutFeaturesToEnable->Core_1_0.shaderResourceMinLod = VK_FALSE;
-	InOutFeaturesToEnable->Core_1_0.sparseBinding = VK_FALSE;
-	InOutFeaturesToEnable->Core_1_0.sparseResidencyBuffer = VK_FALSE;
-	InOutFeaturesToEnable->Core_1_0.sparseResidencyImage2D = VK_FALSE;
-	InOutFeaturesToEnable->Core_1_0.sparseResidencyImage3D = VK_FALSE;
-	InOutFeaturesToEnable->Core_1_0.sparseResidency2Samples = VK_FALSE;
-	InOutFeaturesToEnable->Core_1_0.sparseResidency4Samples = VK_FALSE;
-	InOutFeaturesToEnable->Core_1_0.sparseResidency8Samples = VK_FALSE;
-	InOutFeaturesToEnable->Core_1_0.sparseResidencyAliased = VK_FALSE;
+	if (CVarVulkanAllowSparseResidency.GetValueOnAnyThread() == 0)
+	{
+		// Disable everything sparse-related
+		InOutFeaturesToEnable->Core_1_0.shaderResourceResidency = VK_FALSE;
+		InOutFeaturesToEnable->Core_1_0.shaderResourceMinLod = VK_FALSE;
+		InOutFeaturesToEnable->Core_1_0.sparseBinding = VK_FALSE;
+		InOutFeaturesToEnable->Core_1_0.sparseResidencyBuffer = VK_FALSE;
+		InOutFeaturesToEnable->Core_1_0.sparseResidencyImage2D = VK_FALSE;
+		InOutFeaturesToEnable->Core_1_0.sparseResidencyImage3D = VK_FALSE;
+		InOutFeaturesToEnable->Core_1_0.sparseResidency2Samples = VK_FALSE;
+		InOutFeaturesToEnable->Core_1_0.sparseResidency4Samples = VK_FALSE;
+		InOutFeaturesToEnable->Core_1_0.sparseResidency8Samples = VK_FALSE;
+		InOutFeaturesToEnable->Core_1_0.sparseResidencyAliased = VK_FALSE;
+	}
 }
 
 VkResult FVulkanGenericPlatform::Present(VkQueue Queue, VkPresentInfoKHR& PresentInfo)
diff --git a/Engine/Source/Runtime/VulkanRHI/Private/VulkanQueue.cpp b/Engine/Source/Runtime/VulkanRHI/Private/VulkanQueue.cpp
--- a/Engine/Source/Runtime/VulkanRHI/Private/VulkanQueue.cpp
+++ b/Engine/Source/Runtime/VulkanRHI/Private/VulkanQueue.cpp
@@ -95,6 +95,78 @@ FVulkanQueue::~FVulkanQueue()
 	}
 }
 
+
+void FVulkanQueue::BindSparseResources(FVulkanPayload& Payload)
+{
+	TArray<TArray<VkSparseMemoryBind>, TInlineAllocator<1>> SparseMemoryBinds;
+	TArray<VkSparseBufferMemoryBindInfo, TInlineAllocator<1>> BufferMemoryBindInfo;
+	TArray<VkSparseImageOpaqueMemoryBindInfo, TInlineAllocator<1>> ImageMemoryBindInfo;
+	for (const FVulkanCommitReservedResourceDesc& CommitDesc : Payload.ReservedResourcesToCommit)
+	{
+		if (CommitDesc.Resource->GetType() == RRT_Buffer)
+		{
+			FVulkanBuffer* Buffer = ResourceCast(static_cast<FRHIBuffer*>(CommitDesc.Resource));
+			TArray<VkSparseMemoryBind> BufferSparseMemoryBinds = Buffer->CommitReservedResource(CommitDesc.CommitSizeInBytes);
+			if (BufferSparseMemoryBinds.Num())
+			{
+				VkSparseBufferMemoryBindInfo BindInfo;
+				BindInfo.buffer = Buffer->GetHandle();
+				BindInfo.bindCount = BufferSparseMemoryBinds.Num();
+				BindInfo.pBinds = BufferSparseMemoryBinds.GetData();
+				BufferMemoryBindInfo.Emplace(BindInfo);
+
+				SparseMemoryBinds.Emplace(MoveTemp(BufferSparseMemoryBinds));
+			}
+		}
+		else if (CommitDesc.Resource->GetType() == RRT_Texture)
+		{
+			FVulkanTexture* Texture = ResourceCast(static_cast<FRHITexture*>(CommitDesc.Resource));
+			TArray<VkSparseMemoryBind> ImageSparseMemoryBinds = Texture->CommitReservedResource(CommitDesc.CommitSizeInBytes);
+			if (ImageSparseMemoryBinds.Num())
+			{
+				VkSparseImageOpaqueMemoryBindInfo BindInfo;
+				BindInfo.image = Texture->Image;
+				BindInfo.bindCount = ImageSparseMemoryBinds.Num();
+				BindInfo.pBinds = ImageSparseMemoryBinds.GetData();
+				ImageMemoryBindInfo.Emplace(BindInfo);
+
+				SparseMemoryBinds.Emplace(MoveTemp(ImageSparseMemoryBinds));
+			}
+		}
+		else
+		{
+			checkNoEntry();
+		}
+	}
+
+	if (SparseMemoryBinds.Num())
+	{
+		TArray<VkSemaphore> SemaphoreHandles;
+		SemaphoreHandles.Reserve(Payload.WaitSemaphores.Num());
+		for (VulkanRHI::FSemaphore* Semaphore : Payload.WaitSemaphores)
+		{
+			SemaphoreHandles.Add(Semaphore->GetHandle());
+		}
+		Payload.WaitSemaphores.Reset();
+
+		VkBindSparseInfo BindSparseInfo;
+		ZeroVulkanStruct(BindSparseInfo, VK_STRUCTURE_TYPE_BIND_SPARSE_INFO);
+		BindSparseInfo.waitSemaphoreCount = SemaphoreHandles.Num();
+		BindSparseInfo.pWaitSemaphores = SemaphoreHandles.GetData();
+		BindSparseInfo.bufferBindCount = BufferMemoryBindInfo.Num();
+		BindSparseInfo.pBufferBinds = BufferMemoryBindInfo.GetData();
+		BindSparseInfo.imageOpaqueBindCount = ImageMemoryBindInfo.Num();
+		BindSparseInfo.pImageOpaqueBinds = ImageMemoryBindInfo.GetData();
+		VERIFYVULKANRESULT(VulkanRHI::vkQueueBindSparse(Queue, 1, &BindSparseInfo, VK_NULL_HANDLE));
+
+		if (GWaitForIdleOnSubmit != 0)
+		{
+			VERIFYVULKANRESULT(VulkanRHI::vkDeviceWaitIdle(Device.GetHandle()));
+		}
+	}
+}
+
+
 int32 FVulkanQueue::SubmitQueuedPayloads(TMap<VkSemaphore, FBinarySemaphoreSignalInfo>& SignaledSemas)
 {
 	auto CanBeProcessed = [&SignaledSemas, &Device=Device](FVulkanPayload* Payload)
@@ -239,13 +311,25 @@ void FVulkanQueue::SubmitPayloads(TArrayView<FVulkanPayload*> Payloads, TMap<VkS
 #endif // RHI_NEW_GPU_PROFILER
 
 		// Some payloads have nothing to submit because they are only used to trigger CPU events
-		if (!Payload->WaitSemaphores.Num() && !Payload->CommandBuffers.Num() && !Payload->SignalSemaphores.Num())
+		if (!Payload->WaitSemaphores.Num() && !Payload->CommandBuffers.Num() && !Payload->SignalSemaphores.Num() && !Payload->ReservedResourcesToCommit.Num())
 		{
 			// Consider complete when previous workload is done
 			Payload->TimelineSemaphoreValue = NextTimelineSemaphoreValue - 1;
 			continue;
 		}
 
+		// Bind memory for sparse resources
+		if (Payload->ReservedResourcesToCommit.Num())
+		{
+			if (SubmitInfos.Num() > 0)
+			{
+				Submit(SubmitInfos, nullptr);
+				SubmitInfos.Reset();
+			}
+
+			BindSparseResources(*Payload);
+		}
+
 		VkSubmitInfo& SubmitInfo = SubmitInfos.AddZeroed_GetRef();
 		SubmitInfo.sType = VK_STRUCTURE_TYPE_SUBMIT_INFO;
 
@@ -257,7 +341,7 @@ void FVulkanQueue::SubmitPayloads(TArrayView<FVulkanPayload*> Payloads, TMap<VkS
 			SubmitInfo.pNext = TimelineInfo;
 		}
 
-		// Wait Sempahores
+		// Wait Semaphores
 		if (Payload->WaitSemaphores.Num())
 		{
 			const int32 FirstSempahoreIndex = SemaphoreStorage.Num();
diff --git a/Engine/Source/Runtime/VulkanRHI/Private/VulkanQueue.h b/Engine/Source/Runtime/VulkanRHI/Private/VulkanQueue.h
--- a/Engine/Source/Runtime/VulkanRHI/Private/VulkanQueue.h
+++ b/Engine/Source/Runtime/VulkanRHI/Private/VulkanQueue.h
@@ -125,6 +125,8 @@ private:
 	void SubmitPayloads(TArrayView<FVulkanPayload*> Payloads, TMap<VkSemaphore, FBinarySemaphoreSignalInfo>& InOutSignaledSemas);
 	void Submit(TArrayView<VkSubmitInfo> InSubmitInfos, VulkanRHI::FFence* Fence);
 
+	void BindSparseResources(FVulkanPayload& Payload);
+
 	// Used by submission pipe which holds the proper locks to access this queue
 	FVulkanPayload* GetNextInterruptPayload()
 	{
diff --git a/Engine/Source/Runtime/VulkanRHI/Private/VulkanRHI.cpp b/Engine/Source/Runtime/VulkanRHI/Private/VulkanRHI.cpp
--- a/Engine/Source/Runtime/VulkanRHI/Private/VulkanRHI.cpp
+++ b/Engine/Source/Runtime/VulkanRHI/Private/VulkanRHI.cpp
@@ -870,13 +870,14 @@ void FVulkanDynamicRHI::InitInstance()
 
 		const VkPhysicalDeviceProperties& Props = Device->GetDeviceProperties();
 		const VkPhysicalDeviceLimits& Limits = Device->GetLimits();
+		const FVulkanPhysicalDeviceFeatures& Features = Device->GetPhysicalDeviceFeatures();
 
 		// Initialize the RHI capabilities.
 		GRHINeedsSRVGraphicsNonPixelWorkaround = true;
 		GRHISupportsFirstInstance = true;
 		GRHISupportsDynamicResolution = FVulkanPlatform::SupportsDynamicResolution();
 		GRHISupportsFrameCyclesBubblesRemoval = true;
-		GSupportsDepthBoundsTest = Device->GetPhysicalDeviceFeatures().Core_1_0.depthBounds != 0;
+		GSupportsDepthBoundsTest = Features.Core_1_0.depthBounds != 0;
 		GSupportsRenderTargetFormat_PF_G8 = false;	// #todo-rco
 		GRHISupportsTextureStreaming = true;
 		GRHISupportsGPUTimestampBubblesRemoval = true;
@@ -967,7 +968,7 @@ void FVulkanDynamicRHI::InitInstance()
 		GMaxTextureArrayLayers = Props.limits.maxImageArrayLayers;
 		GRHISupportsBaseVertexIndex = true;
 		GSupportsSeparateRenderTargetBlendState = true;
-		GSupportsDualSrcBlending = Device->GetPhysicalDeviceFeatures().Core_1_0.dualSrcBlend == VK_TRUE;
+		GSupportsDualSrcBlending = Features.Core_1_0.dualSrcBlend == VK_TRUE;
 		GRHISupportsSeparateDepthStencilCopyAccess = Device->SupportsParallelRendering();
 		GRHIGlobals.bSupportsBindless = Device->SupportsBindless() && FDataDrivenShaderPlatformInfo::GetSupportsBindless(GMaxRHIShaderPlatform);
 		GMaxTextureSamplers = (int32)FMath::Min<uint32>(MAX_int32, Props.limits.maxPerStageDescriptorSamplers);
@@ -977,6 +978,23 @@ void FVulkanDynamicRHI::InitInstance()
 		GRHIMaxDispatchThreadGroupsPerDimension.Z = FMath::Min<uint32>(Limits.maxComputeWorkGroupCount[2], 0x7fffffff);
 		GRHISupportsBindingTexArrayPerSlice = true;
 
+		// Check a few extra things to ensure ease of use (currently met by all vendors)
+		const uint32 GraphicsQueueFamilyIndex = Device->GetGraphicsQueue()->GetFamilyIndex();
+		const VkQueueFlags GraphicsQueueFlags = Device->GetQueueFamilyProps()[GraphicsQueueFamilyIndex].queueFlags;
+		const bool bGraphicsQueueSupportsSparseBinding = VKHasAllFlags(GraphicsQueueFlags, VK_QUEUE_SPARSE_BINDING_BIT) &&
+			Device->SupportsParallelRendering(); // resource commits are only handled in the sync2 path
+		if (bGraphicsQueueSupportsSparseBinding)
+		{
+			GRHIGlobals.ReservedResources.Supported =
+				Features.Core_1_0.sparseBinding &&
+				Features.Core_1_0.sparseResidencyBuffer &&
+				Features.Core_1_0.sparseResidencyImage2D &&
+				Props.sparseProperties.residencyNonResidentStrict;
+
+			// :todo-jn: Wait for maintenance9 to be released
+			GRHIGlobals.ReservedResources.SupportsVolumeTextures = false;
+		}
+
 		// Note: While the 2022/2024 profile limits state a minimum of 16, other profiles (even core) go down to 4
 		//       (see https://vulkan.lunarg.com/doc/view/1.3.290.0/windows/profiles_definitions.html)
 		// Since the RHI has historically always supported 8 UAV's, let's leave those specific devices out
@@ -994,7 +1012,7 @@ void FVulkanDynamicRHI::InitInstance()
 		VkShaderStageFlags VulkanDeviceShaderStageBits = VK_SHADER_STAGE_VERTEX_BIT | VK_SHADER_STAGE_FRAGMENT_BIT | VK_SHADER_STAGE_COMPUTE_BIT;
 
 		// optional shader stages
-		if (Device->GetPhysicalDeviceFeatures().Core_1_0.geometryShader)
+		if (Features.Core_1_0.geometryShader)
 		{
 			GVulkanDevicePipelineStageBits |= VK_PIPELINE_STAGE_GEOMETRY_SHADER_BIT;
 			VulkanDeviceShaderStageBits |= VK_SHADER_STAGE_GEOMETRY_BIT;
diff --git a/Engine/Source/Runtime/VulkanRHI/Private/VulkanResources.h b/Engine/Source/Runtime/VulkanRHI/Private/VulkanResources.h
--- a/Engine/Source/Runtime/VulkanRHI/Private/VulkanResources.h
+++ b/Engine/Source/Runtime/VulkanRHI/Private/VulkanResources.h
@@ -577,6 +577,21 @@ private:
 	FVulkanLinkedView* LinkedViews = nullptr;
 };
 
+struct FVulkanReservedResourceData
+{
+	struct FSparseAllocation
+	{
+		VulkanRHI::FVulkanAllocation Allocation;
+		uint64 ResourceOffset = 0;
+	};
+	TArray<FSparseAllocation> Allocations;
+	VkMemoryRequirements MemoryRequirements;
+	uint64 LastBlockSize = 0;
+	uint32 BlockCount = 0;
+	uint32 NumCommittedTiles = 0;
+};
+
+
 enum class EImageOwnerType : uint8
 {
 	None,
@@ -791,6 +806,8 @@ public:
 
 	void SetInitialImageState(FRHICommandListBase& RHICmdList, VkImageLayout InitialLayout, bool bClear, const FClearValueBinding& ClearValueBinding, bool bIsTransientResource);
 
+	TArray<VkSparseMemoryBind> CommitReservedResource(uint64 RequiredCommitSizeInBytes);
+
 private:
 	void InternalMoveSurface(FVulkanDevice& InDevice, FVulkanCommandListContext& Context, VulkanRHI::FVulkanAllocation& DestAllocation);
 
@@ -808,6 +825,8 @@ private:
 	// Views with lifetime tied to this image
 	TMap<uint32, FVulkanView*> InternalViews;
 
+	TUniquePtr<FVulkanReservedResourceData> ReservedResourceData;
+
 protected:
 	EImageOwnerType ImageOwnerType;
 };
@@ -906,6 +925,8 @@ public:
 	}
 #endif
 
+	TArray<VkSparseMemoryBind> CommitReservedResource(uint64 RequiredCommitSizeInBytes);
+
 protected:
 
 	// Will return a new allocation that can be used for this buffer (proper memory type and size)
@@ -925,6 +946,8 @@ protected:
 	FBufferAlloc CurrentBufferAlloc;
 	uint32 LockCounter = 0;
 
+	TUniquePtr<FVulkanReservedResourceData> ReservedResourceData;
+
 	friend class FVulkanCommandListContext;
 	friend struct FRHICommandMultiBufferUnlock;
 };
diff --git a/Engine/Source/Runtime/VulkanRHI/Private/VulkanSubmission.h b/Engine/Source/Runtime/VulkanRHI/Private/VulkanSubmission.h
--- a/Engine/Source/Runtime/VulkanRHI/Private/VulkanSubmission.h
+++ b/Engine/Source/Runtime/VulkanRHI/Private/VulkanSubmission.h
@@ -27,6 +27,16 @@ static FVulkanSyncPointRef CreateVulkanSyncPoint()
 	return FGraphEvent::CreateGraphEvent();
 }
 
+
+
+struct FVulkanCommitReservedResourceDesc
+{
+	FRHIResource* Resource = nullptr;
+	uint64 CommitSizeInBytes = 0;
+};
+
+
+
 class FVulkanPayload
 {
 	friend class FVulkanCommandListContext;
@@ -89,6 +99,9 @@ protected:
 #else
 	bool bEndFrame = false;
 #endif
+
+	// UpdateReservedResources
+	TArray<FVulkanCommitReservedResourceDesc> ReservedResourcesToCommit;
 };
 
 struct FVulkanPlatformCommandList : public IRHIPlatformCommandList, public TArray<FVulkanPayload*>
diff --git a/Engine/Source/Runtime/VulkanRHI/Private/VulkanTexture.cpp b/Engine/Source/Runtime/VulkanRHI/Private/VulkanTexture.cpp
--- a/Engine/Source/Runtime/VulkanRHI/Private/VulkanTexture.cpp
+++ b/Engine/Source/Runtime/VulkanRHI/Private/VulkanTexture.cpp
@@ -40,6 +40,15 @@ static TAutoConsoleVariable<int32> GCVarAllowConcurrentImage(
 	ECVF_ReadOnly
 );
 
+static TAutoConsoleVariable<int32> CVarVulkanSparseImageAllocSizeMB(
+	TEXT("r.Vulkan.SparseImageAllocSizeMB"),
+	16,
+	TEXT("Size of the backing memory blocks for sparse images in megabytes (default 16MB)."),
+	ECVF_ReadOnly
+);
+
+
+
 extern int32 GVulkanLogDefrag;
 
 #if ENABLE_LOW_LEVEL_MEM_TRACKER
@@ -288,6 +297,11 @@ void FVulkanTexture::GenerateImageCreateInfo(
 
 	ImageCreateInfo.flags = (ResourceType == VK_IMAGE_VIEW_TYPE_CUBE || ResourceType == VK_IMAGE_VIEW_TYPE_CUBE_ARRAY) ? VK_IMAGE_CREATE_CUBE_COMPATIBLE_BIT : 0;
 
+	if (EnumHasAnyFlags(UEFlags, TexCreate_ReservedResource))
+	{
+		ImageCreateInfo.flags |= VK_IMAGE_CREATE_SPARSE_BINDING_BIT; // no VK_IMAGE_CREATE_SPARSE_RESIDENCY_BIT yet since only TexCreate_ImmediateCommit is supported
+	}
+
 	const bool bHasUAVFormat = (InDesc.UAVFormat != PF_Unknown && InDesc.UAVFormat != InDesc.Format);
 	const bool bNeedsMutableFormat = (EnumHasAllFlags(UEFlags, TexCreate_SRGB) || (InDesc.Format == PF_R64_UINT) || bHasUAVFormat || !InDesc.AliasableFormats.IsEmpty());
 	if (bNeedsMutableFormat)
@@ -339,7 +353,7 @@ void FVulkanTexture::GenerateImageCreateInfo(
 		}
 	}
 
-	if (ImageCreateInfo.imageType == VK_IMAGE_TYPE_3D)
+	if ((ImageCreateInfo.imageType == VK_IMAGE_TYPE_3D) && !EnumHasAnyFlags(UEFlags, TexCreate_ReservedResource))
 	{
 		ImageCreateInfo.flags |= VK_IMAGE_CREATE_2D_ARRAY_COMPATIBLE_BIT;
 	}
@@ -714,7 +728,18 @@ void FVulkanTexture::DestroySurface()
 			if (Image != VK_NULL_HANDLE)
 			{
 				Device->GetDeferredDeletionQueue().EnqueueResource(VulkanRHI::FDeferredDeletionQueue2::EType::Image, Image);
-				Device->GetMemoryManager().FreeVulkanAllocation(Allocation);
+
+				if (ReservedResourceData)
+				{
+					check(!Allocation.HasAllocation());
+					CommitReservedResource(0);
+					VulkanTextureDestroyed(GetDesc(), MemoryRequirements.size);
+				}
+				else
+				{
+					Device->GetMemoryManager().FreeVulkanAllocation(Allocation);
+				}
+
 				Image = VK_NULL_HANDLE;
 			}
 		}
@@ -881,6 +906,89 @@ void FVulkanTexture::SetInitialImageState(FRHICommandListBase& RHICmdList, VkIma
 		});
 }
 
+TArray<VkSparseMemoryBind> FVulkanTexture::CommitReservedResource(uint64 RequiredCommitSizeInBytes)
+{
+	checkf((RequiredCommitSizeInBytes == UINT64_MAX) || (RequiredCommitSizeInBytes == 0), TEXT("Only full resource commits are supported for images."));
+	checkf(ReservedResourceData, TEXT("CommitReservedResource called on a resource without any ReservedResourceData.  Resource needs to be created with the ReservedResource flag."));
+
+	const uint64 SparseBlockSize = (uint64)CVarVulkanSparseImageAllocSizeMB.GetValueOnAnyThread() * 1024 * 1024;
+	const uint64 TileSize = GRHIGlobals.ReservedResources.TileSizeInBytes;
+	RequiredCommitSizeInBytes = Align(FMath::Min<uint64>(RequiredCommitSizeInBytes, ReservedResourceData->MemoryRequirements.size), TileSize);
+
+	// Figure out the allocation needs
+	const uint64 RequiredBlockCount = Align(RequiredCommitSizeInBytes, SparseBlockSize) / SparseBlockSize;
+	const uint64 CurrentBlockCount = ReservedResourceData->Allocations.Num();
+
+	// Figure out the tile needs
+	const uint64 RequiredTileCount = Align(RequiredCommitSizeInBytes, TileSize) / TileSize;
+	const uint64 CurrentTileCount = ReservedResourceData->NumCommittedTiles;
+
+	TArray<VkSparseMemoryBind> Binds;
+	if (RequiredTileCount > CurrentTileCount)
+	{
+		VkMemoryRequirements BlockMemoryRequirements = ReservedResourceData->MemoryRequirements;
+		BlockMemoryRequirements.size = SparseBlockSize;
+
+		// Make sure we will have all the memory we need
+		ReservedResourceData->Allocations.Reserve(RequiredBlockCount);
+		for (uint64 BlockCount = CurrentBlockCount; BlockCount < RequiredBlockCount; BlockCount++)
+		{
+			const uint32 BlockSize = (BlockCount == ReservedResourceData->BlockCount - 1) ? ReservedResourceData->LastBlockSize : SparseBlockSize;
+			FVulkanReservedResourceData::FSparseAllocation& NewBlockAlloc = ReservedResourceData->Allocations.AddDefaulted_GetRef();
+			if (!Device->GetMemoryManager().AllocateImageMemory(
+				NewBlockAlloc.Allocation, nullptr, BlockMemoryRequirements,
+				VK_MEMORY_PROPERTY_DEVICE_LOCAL_BIT, VulkanRHI::EVulkanAllocationMetaImageOther,
+				false, __FILE__, __LINE__))
+			{
+				Device->GetMemoryManager().HandleOOM();
+			}
+			NewBlockAlloc.ResourceOffset = BlockCount * SparseBlockSize;
+		}
+
+		// Bind all the individual tiles
+		Binds.Reserve(RequiredTileCount - CurrentTileCount);
+		for (uint64 TileCount = CurrentTileCount; TileCount < RequiredTileCount; TileCount++)
+		{
+			const uint32 TileOffset = (TileCount * TileSize);
+
+			const uint32 AllocIndex = TileOffset / SparseBlockSize;
+			const FVulkanReservedResourceData::FSparseAllocation& SparseAlloc = ReservedResourceData->Allocations[AllocIndex];
+
+			VkSparseMemoryBind& NewBind = Binds.AddZeroed_GetRef();
+			NewBind.size = TileSize;
+			NewBind.memory = SparseAlloc.Allocation.GetDeviceMemoryHandle(Device);
+			NewBind.memoryOffset = SparseAlloc.Allocation.Offset + TileOffset - SparseAlloc.ResourceOffset;
+			NewBind.resourceOffset = TileOffset;
+		}
+
+		const int64 CommitDeltaInBytes = TileSize * Binds.Num();
+		UE::RHICore::UpdateReservedResourceStatsOnCommit(CommitDeltaInBytes, false /* bBuffer */, true /* bCommitting */);
+	}
+	else if (RequiredTileCount < CurrentTileCount)
+	{
+		for (int64 TileIndex = CurrentTileCount - 1; TileIndex >= (int64)RequiredTileCount; TileIndex--)
+		{
+			VkSparseMemoryBind& NewBind = Binds.AddZeroed_GetRef();
+			NewBind.size = TileSize;
+			NewBind.memory = VK_NULL_HANDLE;
+			NewBind.memoryOffset = 0;
+			NewBind.resourceOffset = (TileIndex * TileSize);
+		}
+
+		// Free up unused blocks (deletion is deferred, don't have to wait for binds to be submitted)
+		while (ReservedResourceData->Allocations.Num() > RequiredBlockCount)
+		{
+			ReservedResourceData->Allocations.Last().Allocation.Free(*Device);
+			ReservedResourceData->Allocations.Pop();
+		}
+
+		const int64 CommitDeltaInBytes = TileSize * Binds.Num();
+		UE::RHICore::UpdateReservedResourceStatsOnCommit(CommitDeltaInBytes, false /* bBuffer */, false /* bCommitting */);
+	}
+
+	ReservedResourceData->NumCommittedTiles = RequiredTileCount;
+	return Binds;
+}
 
 
 /*-----------------------------------------------------------------------------
@@ -1048,6 +1156,22 @@ FVulkanTexture* FVulkanDynamicRHI::FinalizeCreateTextureInternal(FRHICommandList
 	FVulkanTexture* Texture = CreateResult.Texture;
 	const FRHITextureDesc& Desc = Texture->GetDesc();
 
+	if (EnumHasAnyFlags(Desc.Flags, TexCreate_ImmediateCommit))
+	{
+		checkf(EnumHasAnyFlags(Desc.Flags, TexCreate_ReservedResource), TEXT("The ImmediateCommit flag can only be used on textures created with the ReservedResource flag."));
+
+		RHICmdList.EnqueueLambda(TEXT("FVulkanTexture::ImmediateCommit"),
+			[Texture](FRHICommandListBase& ExecutingCmdList)
+			{
+				FVulkanCommandListContext& Context = FVulkanCommandListContext::Get(ExecutingCmdList);
+
+				FVulkanCommitReservedResourceDesc CommitDesc;
+				CommitDesc.Resource = Texture;
+				CommitDesc.CommitSizeInBytes = UINT64_MAX;
+				Context.GetPayload(FVulkanCommandListContext::EPhase::UpdateReservedResources).ReservedResourcesToCommit.Add(CommitDesc);
+			});
+	}
+
 	if (!EnumHasAnyFlags(Desc.Flags, TexCreate_CPUReadback))
 	{
 		if (CreateResult.DefaultLayout != VK_IMAGE_LAYOUT_UNDEFINED || CreateResult.bClear)
@@ -1188,6 +1312,7 @@ FVulkanTexture* FVulkanDynamicRHI::CreateTextureInternal(FRHICommandListBase& RH
 
 FVulkanTexture* FVulkanDynamicRHI::CreateTextureInternal(const FRHITextureCreateDesc& CreateDesc, const FRHITransientHeapAllocation& InTransientHeapAllocation)
 {
+	checkf(!EnumHasAnyFlags(CreateDesc.Flags, TexCreate_ReservedResource), TEXT("Can't use reserved resources with transient allocation."));
 	FCreateTextureResult CreateResult = BeginCreateTextureInternal(CreateDesc, &InTransientHeapAllocation);
 	return CreateResult.Texture;
 }
@@ -1646,6 +1771,7 @@ FVulkanTexture::FVulkanTexture(FVulkanDevice& InDevice, const FRHITextureCreateD
 	const bool bRenderTarget = EnumHasAnyFlags(InCreateDesc.Flags, TexCreate_RenderTargetable | TexCreate_DepthStencilTargetable | TexCreate_ResolveTargetable);
 	const bool bUAV = EnumHasAnyFlags(InCreateDesc.Flags, TexCreate_UAV);
 	const bool bExternal = EnumHasAnyFlags(InCreateDesc.Flags, TexCreate_External);
+	const bool bIsTransientResource = (InTransientHeapAllocation != nullptr) && InTransientHeapAllocation->IsValid();
 
 	VkMemoryPropertyFlags MemoryFlags = VK_MEMORY_PROPERTY_DEVICE_LOCAL_BIT;
 
@@ -1662,41 +1788,70 @@ FVulkanTexture::FVulkanTexture(FVulkanDevice& InDevice, const FRHITextureCreateD
 		}
 	}
 
-	const bool bIsTransientResource = (InTransientHeapAllocation != nullptr) && InTransientHeapAllocation->IsValid();
-	if (bIsTransientResource)
+	if (EnumHasAnyFlags(InCreateDesc.Flags, TexCreate_ReservedResource))
 	{
-		check(!bMemoryless);
-		check(InTransientHeapAllocation->Offset % MemoryRequirements.alignment == 0);
-		check(InTransientHeapAllocation->Size >= MemoryRequirements.size);
-		Allocation = FVulkanTransientHeap::GetVulkanAllocation(*InTransientHeapAllocation);
+		check(!bIsTransientResource);
+		ReservedResourceData = MakeUnique<FVulkanReservedResourceData>();
+		ReservedResourceData->MemoryRequirements = MemoryRequirements;
+
+		const uint64 SparseBlockSize = (uint64)CVarVulkanSparseImageAllocSizeMB.GetValueOnAnyThread() * 1024 * 1024;
+		ReservedResourceData->LastBlockSize = Align(MemoryRequirements.size % SparseBlockSize, GRHIGlobals.ReservedResources.TileSizeInBytes);
+		ReservedResourceData->BlockCount = MemoryRequirements.size / SparseBlockSize;
+		if (ReservedResourceData->LastBlockSize == 0)
+		{
+			ReservedResourceData->LastBlockSize = SparseBlockSize;
+		}
+		else
+		{
+			ReservedResourceData->BlockCount++;
+		}
+
+		checkf(SparseBlockSize % ReservedResourceData->MemoryRequirements.alignment == 0,
+			TEXT("The value of r.Vulkan.SparseImageAllocSizeMB (%dMB) is not compatible with this resource's tile size requirements (%d)."),
+			CVarVulkanSparseImageAllocSizeMB.GetValueOnAnyThread(), ReservedResourceData->MemoryRequirements.alignment);
+		checkf(ReservedResourceData->LastBlockSize % ReservedResourceData->MemoryRequirements.alignment == 0,
+			TEXT("The value of GRHIGlobals.ReservedResources.TileSizeInBytes (%d) is not compatible with this resource's tile size requirements (%d)."),
+			GRHIGlobals.ReservedResources.TileSizeInBytes, ReservedResourceData->MemoryRequirements.alignment);
+
+		VulkanTextureAllocated(GetDesc(), MemoryRequirements.size);
 	}
 	else
 	{
-		VulkanRHI::EVulkanAllocationMetaType MetaType = (bRenderTarget || bUAV) ? VulkanRHI::EVulkanAllocationMetaImageRenderTarget : VulkanRHI::EVulkanAllocationMetaImageOther;
-#if VULKAN_SUPPORTS_DEDICATED_ALLOCATION
-		extern int32 GVulkanEnableDedicatedImageMemory;
-		// Per https://developer.nvidia.com/what%E2%80%99s-your-vulkan-memory-type
-		VkDeviceSize SizeToBeConsideredForDedicated = 12 * 1024 * 1024;
-		if ((bRenderTarget || MemoryRequirements.size >= SizeToBeConsideredForDedicated) && !bMemoryless && GVulkanEnableDedicatedImageMemory)
+		if (bIsTransientResource)
 		{
-			if (!InDevice.GetMemoryManager().AllocateDedicatedImageMemory(Allocation, this, Image, MemoryRequirements, MemoryFlags, MetaType, bExternal, __FILE__, __LINE__))
-			{
-				checkNoEntry();
-			}
+			check(!bMemoryless);
+			check(InTransientHeapAllocation->Offset % MemoryRequirements.alignment == 0);
+			check(InTransientHeapAllocation->Size >= MemoryRequirements.size);
+			Allocation = FVulkanTransientHeap::GetVulkanAllocation(*InTransientHeapAllocation);
 		}
 		else
-#endif
 		{
-			if (!InDevice.GetMemoryManager().AllocateImageMemory(Allocation, this, MemoryRequirements, MemoryFlags, MetaType, bExternal, __FILE__, __LINE__))
+			VulkanRHI::EVulkanAllocationMetaType MetaType = (bRenderTarget || bUAV) ? VulkanRHI::EVulkanAllocationMetaImageRenderTarget : VulkanRHI::EVulkanAllocationMetaImageOther;
+	#if VULKAN_SUPPORTS_DEDICATED_ALLOCATION
+			extern int32 GVulkanEnableDedicatedImageMemory;
+			// Per https://developer.nvidia.com/what%E2%80%99s-your-vulkan-memory-type
+			VkDeviceSize SizeToBeConsideredForDedicated = 12 * 1024 * 1024;
+			if ((bRenderTarget || MemoryRequirements.size >= SizeToBeConsideredForDedicated) && !bMemoryless && GVulkanEnableDedicatedImageMemory)
 			{
-				checkNoEntry();
+				if (!InDevice.GetMemoryManager().AllocateDedicatedImageMemory(Allocation, this, Image, MemoryRequirements, MemoryFlags, MetaType, bExternal, __FILE__, __LINE__))
+				{
+					checkNoEntry();
+				}
+			}
+			else
+	#endif
+			{
+				if (!InDevice.GetMemoryManager().AllocateImageMemory(Allocation, this, MemoryRequirements, MemoryFlags, MetaType, bExternal, __FILE__, __LINE__))
+				{
+					checkNoEntry();
+				}
 			}
-		}
 
-		// update rhi stats
-		VulkanTextureAllocated(GetDesc(), Allocation.Size);
+			// update rhi stats
+			VulkanTextureAllocated(GetDesc(), Allocation.Size);
+		}
+		Allocation.BindImage(Device, Image);
 	}
-	Allocation.BindImage(Device, Image);
 
 	Tiling = ImageCreateInfo.ImageCreateInfo.tiling;
 	check(Tiling == VK_IMAGE_TILING_LINEAR || Tiling == VK_IMAGE_TILING_OPTIMAL);
