VulkanRHI: Remove unused FGPUEvent class in VulkanRHI namespace.

#rnx

diff --git a/Engine/Source/Runtime/VulkanRHI/Private/VulkanMemory.cpp b/Engine/Source/Runtime/VulkanRHI/Private/VulkanMemory.cpp
--- a/Engine/Source/Runtime/VulkanRHI/Private/VulkanMemory.cpp
+++ b/Engine/Source/Runtime/VulkanRHI/Private/VulkanMemory.cpp
@@ -4701,20 +4701,6 @@ namespace VulkanRHI
 	}
 
 
-	FGPUEvent::FGPUEvent(FVulkanDevice& InDevice)
-		: Device(InDevice)
-	{
-		VkEventCreateInfo Info;
-		ZeroVulkanStruct(Info, VK_STRUCTURE_TYPE_EVENT_CREATE_INFO);
-		VERIFYVULKANRESULT(VulkanRHI::vkCreateEvent(InDevice.GetHandle(), &Info, VULKAN_CPU_ALLOCATOR, &Handle));
-	}
-
-	FGPUEvent::~FGPUEvent()
-	{
-		Device.GetDeferredDeletionQueue().EnqueueResource(VulkanRHI::FDeferredDeletionQueue2::EType::Event, Handle);
-	}
-
-
 
 	FDeferredDeletionQueue2::FDeferredDeletionQueue2(FVulkanDevice& InDevice)
 		: Device(InDevice)
diff --git a/Engine/Source/Runtime/VulkanRHI/Private/VulkanMemory.h b/Engine/Source/Runtime/VulkanRHI/Private/VulkanMemory.h
--- a/Engine/Source/Runtime/VulkanRHI/Private/VulkanMemory.h
+++ b/Engine/Source/Runtime/VulkanRHI/Private/VulkanMemory.h
@@ -1083,23 +1083,6 @@ namespace VulkanRHI
 		void DestroyFence(FFence* Fence);
 	};
 
-	class FGPUEvent : public FRefCount
-	{
-	public:
-		FGPUEvent(FVulkanDevice& InDevice);
-		virtual ~FGPUEvent();
-
-		inline VkEvent GetHandle() const
-		{
-			return Handle;
-		}
-
-	protected:
-		FVulkanDevice& Device;
-		VkEvent Handle;
-	};
-
-
 	// Simple temp allocation blocks used for volatile allocations
 	class FTempBlockAllocator
 	{
