VulkanRHI: Remove Win10 requirement for LocalHostVisible memory.

#rnx

diff --git a/Engine/Source/Runtime/VulkanRHI/Private/Windows/VulkanWindowsPlatform.cpp b/Engine/Source/Runtime/VulkanRHI/Private/Windows/VulkanWindowsPlatform.cpp
--- a/Engine/Source/Runtime/VulkanRHI/Private/Windows/VulkanWindowsPlatform.cpp
+++ b/Engine/Source/Runtime/VulkanRHI/Private/Windows/VulkanWindowsPlatform.cpp
@@ -191,6 +191,5 @@ void FVulkanWindowsPlatform::CreateSurface(FVulkanPlatformWindowContext& WindowC
 
 bool FVulkanWindowsPlatform::SupportsDeviceLocalHostVisibleWithNoPenalty(EGpuVendorId VendorId)
 {
-	static bool bIsWin10 = FPlatformMisc::VerifyWindowsVersion(10, 0) /*Win10*/;
-	return (VendorId == EGpuVendorId::Amd && bIsWin10);
+	return (VendorId == EGpuVendorId::Amd);
 }
