VulkanRHI: Move Fences and Semaphores out of VulkanRHI namespace (add Vulkan to their names).  Replace custom FRefCount class with the engine's FThreadSafeRefCountedObject.  Move all sync classes to new VulkanSynchronization.h/cpp.

#rnx

diff --git a/Engine/Source/Runtime/VulkanRHI/Private/VulkanBarriers.cpp b/Engine/Source/Runtime/VulkanRHI/Private/VulkanBarriers.cpp
--- a/Engine/Source/Runtime/VulkanRHI/Private/VulkanBarriers.cpp
+++ b/Engine/Source/Runtime/VulkanRHI/Private/VulkanBarriers.cpp
@@ -659,7 +659,7 @@ void FVulkanDynamicRHI::RHICreateTransition(FRHITransition* Transition, const FR
 
 	if ((SrcPipelines != DstPipelines) && !EnumHasAnyFlags(CreateInfo.Flags, ERHITransitionCreateFlags::NoFence))
 	{
-		Data->Semaphore = new VulkanRHI::FSemaphore(*Device);
+		Data->Semaphore = new FVulkanSemaphore(*Device);
 	}
 
 	// If we're staying on the same queue, use split barriers if they are permitted and supported
diff --git a/Engine/Source/Runtime/VulkanRHI/Private/VulkanBarriers.h b/Engine/Source/Runtime/VulkanRHI/Private/VulkanBarriers.h
--- a/Engine/Source/Runtime/VulkanRHI/Private/VulkanBarriers.h
+++ b/Engine/Source/Runtime/VulkanRHI/Private/VulkanBarriers.h
@@ -4,8 +4,8 @@
 
 #include "Containers/Array.h"
 #include "RHIPipeline.h"
-#include "VulkanMemory.h"
 #include "VulkanResources.h"
+#include "VulkanSynchronization.h"
 #include "VulkanThirdParty.h"
 
 class FVulkanCommandBuffer;
@@ -39,7 +39,7 @@ struct FVulkanPipelineBarrier
 struct FVulkanTransitionData
 {
 	ERHIPipeline SrcPipelines, DstPipelines;
-	VulkanRHI::FSemaphore* Semaphore = nullptr;  // used for cross queue
+	FVulkanSemaphore* Semaphore = nullptr;  // used for cross queue
 	VkEvent EventHandle = VK_NULL_HANDLE;  // used for split barriers
 	TArray<FRHITransitionInfo> TransitionInfos;
 };
diff --git a/Engine/Source/Runtime/VulkanRHI/Private/VulkanContext.h b/Engine/Source/Runtime/VulkanRHI/Private/VulkanContext.h
--- a/Engine/Source/Runtime/VulkanRHI/Private/VulkanContext.h
+++ b/Engine/Source/Runtime/VulkanRHI/Private/VulkanContext.h
@@ -10,6 +10,7 @@
 #include "VulkanRHIPrivate.h"
 #include "VulkanGPUProfiler.h"
 #include "VulkanSubmission.h"
+#include "VulkanSynchronization.h"
 
 class FVulkanDevice;
 class FVulkanDynamicRHI;
@@ -99,11 +100,11 @@ public:
 		return *Payload.CommandBuffers.Last();
 	}
 
-	void AddWaitSemaphore(VkPipelineStageFlags InWaitFlags, VulkanRHI::FSemaphore* InWaitSemaphore)
+	void AddWaitSemaphore(VkPipelineStageFlags InWaitFlags, FVulkanSemaphore* InWaitSemaphore)
 	{
-		AddWaitSemaphores(InWaitFlags, MakeArrayView<VulkanRHI::FSemaphore*>(&InWaitSemaphore, 1));
+		AddWaitSemaphores(InWaitFlags, MakeArrayView<FVulkanSemaphore*>(&InWaitSemaphore, 1));
 	}
-	void AddWaitSemaphores(VkPipelineStageFlags InWaitFlags, TArrayView<VulkanRHI::FSemaphore*> InWaitSemaphores)
+	void AddWaitSemaphores(VkPipelineStageFlags InWaitFlags, TArrayView<FVulkanSemaphore*> InWaitSemaphores)
 	{
 		if (InWaitSemaphores.Num())
 		{
@@ -120,11 +121,11 @@ public:
 		}
 	}
 
-	void AddSignalSemaphore(VulkanRHI::FSemaphore* InSignalSemaphore)
+	void AddSignalSemaphore(FVulkanSemaphore* InSignalSemaphore)
 	{
-		AddSignalSemaphores(MakeArrayView<VulkanRHI::FSemaphore*>(&InSignalSemaphore, 1));
+		AddSignalSemaphores(MakeArrayView<FVulkanSemaphore*>(&InSignalSemaphore, 1));
 	}
-	void AddSignalSemaphores(TArrayView<VulkanRHI::FSemaphore*> InSignalSemaphores)
+	void AddSignalSemaphores(TArrayView<FVulkanSemaphore*> InSignalSemaphores)
 	{
 		if (InSignalSemaphores.Num())
 		{
diff --git a/Engine/Source/Runtime/VulkanRHI/Private/VulkanDevice.cpp b/Engine/Source/Runtime/VulkanRHI/Private/VulkanDevice.cpp
--- a/Engine/Source/Runtime/VulkanRHI/Private/VulkanDevice.cpp
+++ b/Engine/Source/Runtime/VulkanRHI/Private/VulkanDevice.cpp
@@ -231,6 +231,7 @@ FVulkanDevice::FVulkanDevice(FVulkanDynamicRHI* InRHI, VkPhysicalDevice InGpu)
 	: Device(VK_NULL_HANDLE)
 	, MemoryManager(*this)
 	, DeferredDeletionQueue(*this)
+	, FenceManager(*this)
 	, Gpu(InGpu)
 	, ImmediateContext(nullptr)
 	, PipelineStateCache(nullptr)
@@ -1346,8 +1347,6 @@ void FVulkanDevice::InitGPU()
 
 	MemoryManager.Init();
 
-	FenceManager.Init(this);
-
 	StagingManager.Init(this);
 
 	// TempBlockAllocator is currently used for UB uploads and copies
diff --git a/Engine/Source/Runtime/VulkanRHI/Private/VulkanDevice.h b/Engine/Source/Runtime/VulkanRHI/Private/VulkanDevice.h
--- a/Engine/Source/Runtime/VulkanRHI/Private/VulkanDevice.h
+++ b/Engine/Source/Runtime/VulkanRHI/Private/VulkanDevice.h
@@ -7,6 +7,7 @@
 #pragma once
 
 #include "VulkanMemory.h"
+#include "VulkanSynchronization.h"
 #include "VulkanResources.h"
 #include "VulkanQueue.h"
 #include "VulkanQuery.h"
@@ -25,6 +26,7 @@ class FVulkanRenderPassManager;
 class FVulkanRayTracingCompactionRequestHandler;
 class FVulkanRayTracingPipelineLibraryCache;
 struct FVulkanParallelRenderPassInfo;
+class FVulkanFenceManager;
 
 #define VULKAN_USE_DEBUG_NAMES 1
 #if VULKAN_USE_DEBUG_NAMES
@@ -456,7 +458,7 @@ public:
 		return StagingManager;
 	}
 
-	VulkanRHI::FFenceManager& GetFenceManager()
+	FVulkanFenceManager& GetFenceManager()
 	{
 		return FenceManager;
 	}
@@ -626,7 +628,7 @@ private:
 
 	VulkanRHI::FStagingManager StagingManager;
 
-	VulkanRHI::FFenceManager FenceManager;
+	FVulkanFenceManager FenceManager;
 
 	VulkanRHI::FTempBlockAllocator* TempBlockAllocator = nullptr;
 
diff --git a/Engine/Source/Runtime/VulkanRHI/Private/VulkanDynamicRHI.h b/Engine/Source/Runtime/VulkanRHI/Private/VulkanDynamicRHI.h
--- a/Engine/Source/Runtime/VulkanRHI/Private/VulkanDynamicRHI.h
+++ b/Engine/Source/Runtime/VulkanRHI/Private/VulkanDynamicRHI.h
@@ -15,16 +15,13 @@ class FVulkanDevice;
 class FVulkanUniformBuffer;
 class FVulkanViewport;
 class FVulkanPayload;
+class FVulkanSemaphore;
 class FVulkanTexture;
 class IHeadMountedDisplayVulkanExtensions;
 class FVulkanThread;
 struct FRHITransientHeapAllocation;
 struct FVulkanPlatformCommandList;
 
-namespace VulkanRHI
-{
-	class FSemaphore;
-};
 
 struct FOptionalVulkanInstanceExtensions
 {
@@ -368,7 +365,7 @@ protected:
 	FVulkanThread* SubmissionThread = nullptr;
 	FCriticalSection InterruptCS;
 	FVulkanThread* InterruptThread = nullptr;
-	VulkanRHI::FSemaphore* CPUTimelineSemaphore = nullptr;  // used to wake up interrupt thread from CPU
+	FVulkanSemaphore* CPUTimelineSemaphore = nullptr;  // used to wake up interrupt thread from CPU
 	std::atomic<uint64> CPUTimelineSemaphoreValue = 1;
 	TQueue<FVulkanPlatformCommandList*, EQueueMode::Mpsc> PendingPayloadsForSubmission;
 	TMap<VkSemaphore, FBinarySemaphoreSignalInfo> SignaledSemaphores;
diff --git a/Engine/Source/Runtime/VulkanRHI/Private/VulkanMemory.cpp b/Engine/Source/Runtime/VulkanRHI/Private/VulkanMemory.cpp
--- a/Engine/Source/Runtime/VulkanRHI/Private/VulkanMemory.cpp
+++ b/Engine/Source/Runtime/VulkanRHI/Private/VulkanMemory.cpp
@@ -4520,188 +4520,6 @@ namespace VulkanRHI
 		ProcessPendingFreeNoLock(bImmediately, bFreeToOS);
 	}
 
-	FFence::FFence(FVulkanDevice* InDevice, FFenceManager* InOwner, bool bCreateSignaled)
-		: State(bCreateSignaled ? FFence::EState::Signaled : FFence::EState::NotReady)
-		, Owner(InOwner)
-	{
-		VkFenceCreateInfo Info;
-		ZeroVulkanStruct(Info, VK_STRUCTURE_TYPE_FENCE_CREATE_INFO);
-		Info.flags = bCreateSignaled ? VK_FENCE_CREATE_SIGNALED_BIT : 0;
-		VERIFYVULKANRESULT(VulkanRHI::vkCreateFence(InDevice->GetHandle(), &Info, VULKAN_CPU_ALLOCATOR, &Handle));
-	}
-
-	FFence::~FFence()
-	{
-		checkf(Handle == VK_NULL_HANDLE, TEXT("Didn't get properly destroyed by FFenceManager!"));
-	}
-
-	FFenceManager::~FFenceManager()
-	{
-		ensure(UsedFences.Num() == 0);
-	}
-
-	inline void FFenceManager::DestroyFence(FFence* Fence)
-	{
-		// Does not need to go in the deferred deletion queue
-		VulkanRHI::vkDestroyFence(Device->GetHandle(), Fence->GetHandle(), VULKAN_CPU_ALLOCATOR);
-		Fence->Handle = VK_NULL_HANDLE;
-		delete Fence;
-	}
-
-	void FFenceManager::Init(FVulkanDevice* InDevice)
-	{
-		Device = InDevice;
-	}
-
-	void FFenceManager::Deinit()
-	{
-		FScopeLock Lock(&FenceLock);
-		ensureMsgf(UsedFences.Num() == 0, TEXT("No all fences are done!"));
-		VkDevice DeviceHandle = Device->GetHandle();
-		for (FFence* Fence : FreeFences)
-		{
-			DestroyFence(Fence);
-		}
-	}
-
-	FFence* FFenceManager::AllocateFence(bool bCreateSignaled)
-	{
-		FScopeLock Lock(&FenceLock);
-		if (FreeFences.Num() != 0)
-		{
-			FFence* Fence = FreeFences[0];
-			FreeFences.RemoveAtSwap(0, EAllowShrinking::No);
-			UsedFences.Add(Fence);
-
-			if (bCreateSignaled)
-			{
-				Fence->State = FFence::EState::Signaled;
-			}
-			return Fence;
-		}
-
-		FFence* NewFence = new FFence(Device, this, bCreateSignaled);
-		UsedFences.Add(NewFence);
-		return NewFence;
-	}
-
-	// Sets it to nullptr
-	void FFenceManager::ReleaseFence(FFence*& Fence)
-	{
-		FScopeLock Lock(&FenceLock);
-		ResetFence(Fence);
-		UsedFences.RemoveSingleSwap(Fence, EAllowShrinking::No);
-#if VULKAN_REUSE_FENCES
-		FreeFences.Add(Fence);
-#else
-		DestroyFence(Fence);
-#endif
-		Fence = nullptr;
-	}
-
-	void FFenceManager::WaitAndReleaseFence(FFence*& Fence, uint64 TimeInNanoseconds)
-	{
-		FScopeLock Lock(&FenceLock);
-		if (!Fence->IsSignaled())
-		{
-			WaitForFence(Fence, TimeInNanoseconds);
-		}
-
-		ResetFence(Fence);
-		UsedFences.RemoveSingleSwap(Fence, EAllowShrinking::No);
-		FreeFences.Add(Fence);
-		Fence = nullptr;
-	}
-
-	bool FFenceManager::CheckFenceState(FFence* Fence)
-	{
-		check(UsedFences.Contains(Fence));
-		check(Fence->State == FFence::EState::NotReady);
-		VkResult Result = VulkanRHI::vkGetFenceStatus(Device->GetHandle(), Fence->Handle);
-		switch (Result)
-		{
-		case VK_SUCCESS:
-			Fence->State = FFence::EState::Signaled;
-			return true;
-
-		case VK_NOT_READY:
-			break;
-
-		default:
-			VERIFYVULKANRESULT(Result);
-			break;
-		}
-
-		return false;
-	}
-
-	bool FFenceManager::WaitForFence(FFence* Fence, uint64 TimeInNanoseconds)
-	{
-#if VULKAN_ENABLE_AGGRESSIVE_STATS
-		SCOPE_CYCLE_COUNTER(STAT_VulkanWaitFence);
-#endif
-
-		check(UsedFences.Contains(Fence));
-		check(Fence->State == FFence::EState::NotReady);
-		VkResult Result = VulkanRHI::vkWaitForFences(Device->GetHandle(), 1, &Fence->Handle, true, TimeInNanoseconds);
-		switch (Result)
-		{
-		case VK_SUCCESS:
-			Fence->State = FFence::EState::Signaled;
-			return true;
-		case VK_TIMEOUT:
-			break;
-		default:
-			VERIFYVULKANRESULT(Result);
-			break;
-		}
-
-		return false;
-	}
-
-	bool FFenceManager::WaitForAnyFence(TArrayView<FFence*> Fences, uint64 TimeInNanoseconds)
-	{
-		// Mostly used for waits on queues, inlinealloc to that size
-		TArray<VkFence, TInlineAllocator<(int32)EVulkanQueueType::Count>> FenceHandles;
-		for (FFence* Fence : Fences)
-		{
-			if (Fence && !Fence->IsSignaled())
-			{
-				check(UsedFences.Contains(Fence));
-				FenceHandles.Add(Fence->GetHandle());
-			}
-		}
-
-		const VkResult Result = VulkanRHI::vkWaitForFences(Device->GetHandle(), FenceHandles.Num(), FenceHandles.GetData(), false, TimeInNanoseconds);
-		switch (Result)
-		{
-		case VK_SUCCESS:
-			if (Fences.Num() == 1)
-			{
-				Fences[0]->State = FFence::EState::Signaled;
-			}
-			return true;
-		case VK_TIMEOUT:
-			break;
-		default:
-			VERIFYVULKANRESULT(Result);
-			break;
-		}
-
-		return false;
-	}
-
-	void FFenceManager::ResetFence(FFence* Fence)
-	{
-		if (Fence->State != FFence::EState::NotReady)
-		{
-			VERIFYVULKANRESULT(VulkanRHI::vkResetFences(Device->GetHandle(), 1, &Fence->Handle));
-			Fence->State = FFence::EState::NotReady;
-		}
-	}
-
-
-
 	FDeferredDeletionQueue2::FDeferredDeletionQueue2(FVulkanDevice& InDevice)
 		: Device(InDevice)
 	{
@@ -5076,85 +4894,10 @@ namespace VulkanRHI
 			}
 		}
 	}
-
-
-
-
-	FSemaphore::FSemaphore(FVulkanDevice& InDevice, EVulkanSemaphoreFlags InFlags, uint64 InInitialTimelineValue)
-		: Device(InDevice)
-		, SemaphoreHandle(VK_NULL_HANDLE)
-		, Flags(InFlags)
-	{
-		// Create semaphore
-		VkSemaphoreCreateInfo CreateInfo;
-		ZeroVulkanStruct(CreateInfo, VK_STRUCTURE_TYPE_SEMAPHORE_CREATE_INFO);
-
-		VkSemaphoreTypeCreateInfo TypeCreateInfo;
-		if (IsTimeline())
-		{
-			checkf(InDevice.GetOptionalExtensions().HasKHRTimelineSemaphore, TEXT("Trying to create timeline semaphore without device support!"));
-			ZeroVulkanStruct(TypeCreateInfo, VK_STRUCTURE_TYPE_SEMAPHORE_TYPE_CREATE_INFO);
-			TypeCreateInfo.semaphoreType = VK_SEMAPHORE_TYPE_TIMELINE;
-			TypeCreateInfo.initialValue = (uint64_t)InInitialTimelineValue;
-			CreateInfo.pNext = &TypeCreateInfo;
-		}
-
-		VERIFYVULKANRESULT(VulkanRHI::vkCreateSemaphore(Device.GetHandle(), &CreateInfo, VULKAN_CPU_ALLOCATOR, &SemaphoreHandle));
-	}
-
-	FSemaphore::FSemaphore(FVulkanDevice& InDevice, const VkSemaphore& InExternalSemaphore)
-		: Device(InDevice)
-		, SemaphoreHandle(InExternalSemaphore)
-		, Flags(EVulkanSemaphoreFlags::ExternallySignaled)
-	{}
-
-	FSemaphore::~FSemaphore()
-	{
-		check(SemaphoreHandle != VK_NULL_HANDLE);
-		if (EnumHasAnyFlags(Flags, EVulkanSemaphoreFlags::ImmediateDeletion))
-		{
-			VulkanRHI::vkDestroySemaphore(Device.GetHandle(), SemaphoreHandle, VULKAN_CPU_ALLOCATOR);
-		}
-		else if (!IsExternallyOwned())
-		{
-			Device.GetDeferredDeletionQueue().EnqueueResource(VulkanRHI::FDeferredDeletionQueue2::EType::Semaphore, SemaphoreHandle);
-		}
-		SemaphoreHandle = VK_NULL_HANDLE;
-	}
-
-	bool FSemaphore::WaitForTimelineSemaphoreValue(uint64 Value, uint64 Timeout)
-	{
-		checkf(IsTimeline(), TEXT("WaitForTimelineSemaphoreValue not available!  Semaphore was not created with EVulkanSemaphoreFlags::Timeline."));
-		checkSlow(SemaphoreHandle != VK_NULL_HANDLE);
-		VkSemaphoreWaitInfo WaitInfo;
-		ZeroVulkanStruct(WaitInfo, VK_STRUCTURE_TYPE_SEMAPHORE_WAIT_INFO);
-		WaitInfo.semaphoreCount = 1;
-		WaitInfo.pSemaphores = &SemaphoreHandle;
-		WaitInfo.pValues = (uint64_t*)&Value;
-		const VkResult Result = VulkanRHI::vkWaitSemaphoresKHR(Device.GetHandle(), &WaitInfo, Timeout);
-		switch (Result)
-		{
-		case VK_SUCCESS:
-			return true;
-		case VK_TIMEOUT:
-			break;
-		default:
-			VERIFYVULKANRESULT(Result);
-			break;
-		}
-		return false;
-	}
-
-	uint64 FSemaphore::GetTimelineSemaphoreValue()
-	{
-		checkf(IsTimeline(), TEXT("GetTimelineSemaphoreValue not available!  Semaphore was not created with EVulkanSemaphoreFlags::Timeline."));
-		uint64 Value = 0;
-		VERIFYVULKANRESULT(VulkanRHI::vkGetSemaphoreCounterValueKHR(Device.GetHandle(), SemaphoreHandle, (uint64_t*)&Value));
-		return Value;
-	}
 }
 
 
+
 #if VULKAN_CUSTOM_MEMORY_MANAGER_ENABLED
 namespace VulkanRHI
 {
diff --git a/Engine/Source/Runtime/VulkanRHI/Private/VulkanMemory.h b/Engine/Source/Runtime/VulkanRHI/Private/VulkanMemory.h
--- a/Engine/Source/Runtime/VulkanRHI/Private/VulkanMemory.h
+++ b/Engine/Source/Runtime/VulkanRHI/Private/VulkanMemory.h
@@ -8,6 +8,7 @@
 
 #include "Containers/Map.h"
 #include "Misc/ScopeRWLock.h"
+#include "Templates/RefCounting.h"
 #include "VulkanConfiguration.h"
 #include "VulkanThirdParty.h"
 
@@ -170,45 +171,6 @@ namespace VulkanRHI
 	};
 
 
-	// Custom ref counting
-	class FRefCount
-	{
-	public:
-		FRefCount() {}
-		virtual ~FRefCount()
-		{
-			check(NumRefs.GetValue() == 0);
-		}
-
-		inline uint32 AddRef()
-		{
-			int32 NewValue = NumRefs.Increment();
-			check(NewValue > 0);
-			return (uint32)NewValue;
-		}
-
-		inline uint32 Release()
-		{
-			int32 NewValue = NumRefs.Decrement();
-			if (NewValue == 0)
-			{
-				delete this;
-			}
-			check(NewValue >= 0);
-			return (uint32)NewValue;
-		}
-
-		inline uint32 GetRefCount() const
-		{
-			int32 Value = NumRefs.GetValue();
-			check(Value >= 0);
-			return (uint32)Value;
-		}
-
-	private:
-		FThreadSafeCounter NumRefs;
-	};
-
 	// An Allocation off a Device Heap. Lowest level of allocations and bounded by VkPhysicalDeviceLimits::maxMemoryAllocationCount.
 	class FDeviceMemoryAllocation
 	{
@@ -911,7 +873,7 @@ namespace VulkanRHI
 		void ProcessPendingUBFrees(bool bForce);
 	};
 
-	class FStagingBuffer : public FRefCount
+	class FStagingBuffer : public FThreadSafeRefCountedObject
 	{
 	public:
 		FStagingBuffer(FVulkanDevice* InDevice);
@@ -994,95 +956,6 @@ namespace VulkanRHI
 		FVulkanDevice* Device;
 	};
 
-	class FFence
-	{
-	public:
-		FFence(FVulkanDevice* InDevice, FFenceManager* InOwner, bool bCreateSignaled);
-
-		inline VkFence GetHandle() const
-		{
-			return Handle;
-		}
-
-		inline bool IsSignaled() const
-		{
-			return State == EState::Signaled;
-		}
-
-		FFenceManager* GetOwner()
-		{
-			return Owner;
-		}
-
-	protected:
-		VkFence Handle;
-
-		enum class EState
-		{
-			// Initial state
-			NotReady,
-
-			// After GPU processed it
-			Signaled,
-		};
-
-		EState State;
-
-		FFenceManager* Owner;
-
-		// Only owner can delete!
-		~FFence();
-		friend class FFenceManager;
-	};
-
-	class FFenceManager
-	{
-	public:
-		FFenceManager()
-			: Device(nullptr)
-		{
-		}
-		~FFenceManager();
-
-		void Init(FVulkanDevice* InDevice);
-		void Deinit();
-
-		FFence* AllocateFence(bool bCreateSignaled = false);
-
-		inline bool IsFenceSignaled(FFence* Fence)
-		{
-			if (Fence->IsSignaled())
-			{
-				return true;
-			}
-
-			return CheckFenceState(Fence);
-		}
-
-		// Returns false if it timed out
-		bool WaitForAnyFence(TArrayView<FFence*> Fences, uint64 TimeInNanoseconds);
-		bool WaitForFence(FFence* Fence, uint64 TimeInNanoseconds);
-
-		void ResetFence(FFence* Fence);
-
-		// Sets it to nullptr
-		void ReleaseFence(FFence*& Fence);
-
-		// Sets it to nullptr
-		void WaitAndReleaseFence(FFence*& Fence, uint64 TimeInNanoseconds);
-
-	protected:
-		FVulkanDevice* Device;
-		FCriticalSection FenceLock;
-		TArray<FFence*> FreeFences;
-		TArray<FFence*> UsedFences;
-
-		// Returns true if signaled
-		bool CheckFenceState(FFence* Fence);
-
-		void DestroyFence(FFence* Fence);
-	};
-
 	// Simple temp allocation blocks used for volatile allocations
 	class FTempBlockAllocator
 	{
@@ -1130,62 +1003,9 @@ namespace VulkanRHI
 		FRWLock RWLock;
 	};
 
-	enum class EVulkanSemaphoreFlags : uint8
-	{
-		None = 0,
-
-		// Will not delete handle on destruction
-		ExternallyOwned    = 1 << 1,    
-
-		// Inform submission pipeline that the signal will not come from a payload (acquired image semaphore for example)
-		ExternallySignaled = 1 << 2,
 
-		// Create a timeline semaphore (must be supported)
-		Timeline           = 1 << 3,
-
-		// Will not be queued for deferred deletion
-		ImmediateDeletion  = 1 << 4
-	};
-	ENUM_CLASS_FLAGS(EVulkanSemaphoreFlags);
-
-	class FSemaphore : public FRefCount
-	{
-	public:
-		UE_API FSemaphore(FVulkanDevice& InDevice, EVulkanSemaphoreFlags InFlags=EVulkanSemaphoreFlags::None, uint64 InInitialTimelineValue=0);
-		UE_API FSemaphore(FVulkanDevice& InDevice, const VkSemaphore& InExternalSemaphore);
-		UE_API virtual ~FSemaphore();
-
-		VkSemaphore GetHandle() const
-		{
-			return SemaphoreHandle;
-		}
-
-		bool IsExternallyOwned() const
-		{
-			return EnumHasAnyFlags(Flags, EVulkanSemaphoreFlags::ExternallyOwned);
-		}
-
-		bool IsExternallySignaled() const
-		{
-			return EnumHasAnyFlags(Flags, EVulkanSemaphoreFlags::ExternallySignaled);
-		}
-
-		bool IsTimeline() const
-		{
-			return EnumHasAnyFlags(Flags, EVulkanSemaphoreFlags::Timeline);
-		}
-
-		UE_API uint64 GetTimelineSemaphoreValue();
-		UE_API bool WaitForTimelineSemaphoreValue(uint64 Value, uint64 Timeout);
-
-	private:
-		FVulkanDevice& Device;
-		VkSemaphore SemaphoreHandle;
-		const EVulkanSemaphoreFlags Flags;
-	};
 }
 
-
 #if VULKAN_CUSTOM_MEMORY_MANAGER_ENABLED
 struct FVulkanCustomMemManager
 {
diff --git a/Engine/Source/Runtime/VulkanRHI/Private/VulkanQueue.cpp b/Engine/Source/Runtime/VulkanRHI/Private/VulkanQueue.cpp
--- a/Engine/Source/Runtime/VulkanRHI/Private/VulkanQueue.cpp
+++ b/Engine/Source/Runtime/VulkanRHI/Private/VulkanQueue.cpp
@@ -72,8 +72,8 @@ FVulkanQueue::FVulkanQueue(FVulkanDevice& InDevice, uint32 InFamilyIndex, EVulka
 	if (bUseTimelineSemaphores)
 	{
 		// Use ImmediateDeletion for the timeline since it's deleted after the deferrered deletion queue
-		const VulkanRHI::EVulkanSemaphoreFlags SemaphoreFlags = VulkanRHI::EVulkanSemaphoreFlags::Timeline | VulkanRHI::EVulkanSemaphoreFlags::ImmediateDeletion;
-		TimelineSempahore = new VulkanRHI::FSemaphore(InDevice, SemaphoreFlags, CompletedTimelineSemaphoreValue);
+		const EVulkanSemaphoreFlags SemaphoreFlags = EVulkanSemaphoreFlags::Timeline | EVulkanSemaphoreFlags::ImmediateDeletion;
+		TimelineSempahore = new FVulkanSemaphore(InDevice, SemaphoreFlags, CompletedTimelineSemaphoreValue);
 		TimelineSempahore->AddRef();
 	}
 }
@@ -143,7 +143,7 @@ void FVulkanQueue::BindSparseResources(FVulkanPayload& Payload)
 	{
 		TArray<VkSemaphore> SemaphoreHandles;
 		SemaphoreHandles.Reserve(Payload.WaitSemaphores.Num());
-		for (VulkanRHI::FSemaphore* Semaphore : Payload.WaitSemaphores)
+		for (FVulkanSemaphore* Semaphore : Payload.WaitSemaphores)
 		{
 			SemaphoreHandles.Add(Semaphore->GetHandle());
 		}
@@ -171,8 +171,8 @@ int32 FVulkanQueue::SubmitQueuedPayloads(TMap<VkSemaphore, FBinarySemaphoreSigna
 {
 	auto CanBeProcessed = [&SignaledSemas, &Device=Device](FVulkanPayload* Payload)
 	{
-		const TArray<VulkanRHI::FSemaphore*>& WaitSemas = Payload->WaitSemaphores;
-		for (VulkanRHI::FSemaphore* WaitSema : WaitSemas)
+		const TArray<FVulkanSemaphore*>& WaitSemas = Payload->WaitSemaphores;
+		for (FVulkanSemaphore* WaitSema : WaitSemas)
 		{
 			const VkSemaphore SemaphoreHandle = WaitSema->GetHandle();
 			if (!WaitSema->IsExternallySignaled() && !SignaledSemas.Contains(SemaphoreHandle))
@@ -184,7 +184,7 @@ int32 FVulkanQueue::SubmitQueuedPayloads(TMap<VkSemaphore, FBinarySemaphoreSigna
 #if RHI_NEW_GPU_PROFILER
 		// Find the maximum timeline value to wait on for each queue
 		uint64 MaxTimelineWaitValue[(int32)EVulkanQueueType::Count] = {};
-		for (VulkanRHI::FSemaphore* WaitSema : WaitSemas)
+		for (FVulkanSemaphore* WaitSema : WaitSemas)
 		{
 			if (!WaitSema->IsExternallySignaled())
 			{
@@ -211,7 +211,7 @@ int32 FVulkanQueue::SubmitQueuedPayloads(TMap<VkSemaphore, FBinarySemaphoreSigna
 #endif
 
 		// We can only remove them from the list if they are all present
-		for (VulkanRHI::FSemaphore* WaitSema : WaitSemas)
+		for (FVulkanSemaphore* WaitSema : WaitSemas)
 		{
 			if (!WaitSema->IsExternallySignaled())
 			{
@@ -345,7 +345,7 @@ void FVulkanQueue::SubmitPayloads(TArrayView<FVulkanPayload*> Payloads, TMap<VkS
 		if (Payload->WaitSemaphores.Num())
 		{
 			const int32 FirstSempahoreIndex = SemaphoreStorage.Num();
-			for (VulkanRHI::FSemaphore* Semaphore : Payload->WaitSemaphores)
+			for (FVulkanSemaphore* Semaphore : Payload->WaitSemaphores)
 			{
 				SemaphoreStorage.Add(Semaphore->GetHandle());
 			}
@@ -378,7 +378,7 @@ void FVulkanQueue::SubmitPayloads(TArrayView<FVulkanPayload*> Payloads, TMap<VkS
 		if (Payload->SignalSemaphores.Num())
 		{
 			const int32 FirstSempahoreIndex = SemaphoreStorage.Num();
-			for (VulkanRHI::FSemaphore* Semaphore : Payload->SignalSemaphores)
+			for (FVulkanSemaphore* Semaphore : Payload->SignalSemaphores)
 			{
 				SemaphoreStorage.Add(Semaphore->GetHandle());
 				SignaledSemas.Add(Semaphore->GetHandle(), { NextTimelineSemaphoreValue, QueueType });
@@ -447,7 +447,7 @@ void FVulkanQueue::SubmitPayloads(TArrayView<FVulkanPayload*> Payloads, TMap<VkS
 }
 
 
-void FVulkanQueue::Submit(TArrayView<VkSubmitInfo> InSubmitInfos, VulkanRHI::FFence* Fence)
+void FVulkanQueue::Submit(TArrayView<VkSubmitInfo> InSubmitInfos, FVulkanFence* Fence)
 {
 	SCOPE_CYCLE_COUNTER(STAT_VulkanQueueSubmit);
 
diff --git a/Engine/Source/Runtime/VulkanRHI/Private/VulkanQueue.h b/Engine/Source/Runtime/VulkanRHI/Private/VulkanQueue.h
--- a/Engine/Source/Runtime/VulkanRHI/Private/VulkanQueue.h
+++ b/Engine/Source/Runtime/VulkanRHI/Private/VulkanQueue.h
@@ -88,7 +88,7 @@ public:
 		return SupportedAccess;
 	}
 
-	const VulkanRHI::FSemaphore* GetTimelineSemaphore() const
+	const FVulkanSemaphore* GetTimelineSemaphore() const
 	{
 		return TimelineSempahore;
 	}
@@ -123,7 +123,7 @@ public:
 
 private:
 	void SubmitPayloads(TArrayView<FVulkanPayload*> Payloads, TMap<VkSemaphore, FBinarySemaphoreSignalInfo>& InOutSignaledSemas);
-	void Submit(TArrayView<VkSubmitInfo> InSubmitInfos, VulkanRHI::FFence* Fence);
+	void Submit(TArrayView<VkSubmitInfo> InSubmitInfos, FVulkanFence* Fence);
 
 	void BindSparseResources(FVulkanPayload& Payload);
 
@@ -146,7 +146,7 @@ private:
 	TStaticArray<TArray<FVulkanCommandBufferPool*>, (int32)EVulkanCommandBufferType::Count> CommandBufferPools;
 
 	const bool bUseTimelineSemaphores;
-	VulkanRHI::FSemaphore* TimelineSempahore = nullptr;
+	FVulkanSemaphore* TimelineSempahore = nullptr;
 	uint64 NextTimelineSemaphoreValue = 1;
 	uint64 CompletedTimelineSemaphoreValue = 0;
 
diff --git a/Engine/Source/Runtime/VulkanRHI/Private/VulkanSubmission.cpp b/Engine/Source/Runtime/VulkanRHI/Private/VulkanSubmission.cpp
--- a/Engine/Source/Runtime/VulkanRHI/Private/VulkanSubmission.cpp
+++ b/Engine/Source/Runtime/VulkanRHI/Private/VulkanSubmission.cpp
@@ -19,14 +19,14 @@ extern CORE_API bool GIsGPUCrashed;
 
 FVulkanPayload::~FVulkanPayload()
 {
-	for (VulkanRHI::FSemaphore* Semaphore : WaitSemaphores)
+	for (FVulkanSemaphore* Semaphore : WaitSemaphores)
 	{
 		Semaphore->Release();
 	}
 
 	if (Fence)
 	{
-		Fence->GetOwner()->ReleaseFence(Fence);
+		Fence->GetOwner().ReleaseFence(Fence);
 	}
 }
 
@@ -142,7 +142,7 @@ void FVulkanDynamicRHI::InitializeSubmissionPipe()
 	{
 		if (CVarVulkanUseInterruptThread.GetValueOnAnyThread())
 		{
-			CPUTimelineSemaphore = new VulkanRHI::FSemaphore(*Device, VulkanRHI::EVulkanSemaphoreFlags::Timeline | VulkanRHI::EVulkanSemaphoreFlags::ImmediateDeletion, CPUTimelineSemaphoreValue++);
+			CPUTimelineSemaphore = new FVulkanSemaphore(*Device, EVulkanSemaphoreFlags::Timeline | EVulkanSemaphoreFlags::ImmediateDeletion, CPUTimelineSemaphoreValue++);
 			InterruptThread = new FVulkanThread(TEXT("RHIInterruptThread"), TPri_Highest, this, &FVulkanDynamicRHI::WaitAndProcessInterruptQueue);
 		}
 
@@ -401,11 +401,11 @@ bool FVulkanDynamicRHI::WaitAndProcessInterruptQueue()
 			// Figure out if we wait on Fences or TimelineSemaphores
 			if (NextPayloads[0]->Fence)
 			{
-				VulkanRHI::FFence* Fences[NumQueues];
+				FVulkanFence* Fences[NumQueues];
 				int32 NumFences = 0;
 				for (uint32 Index = 0; Index < NumSyncs; ++Index)
 				{
-					VulkanRHI::FFence* Fence = NextPayloads[Index]->Fence;
+					FVulkanFence* Fence = NextPayloads[Index]->Fence;
 					checkf(Fence, TEXT("Payloads should all use the same types of syncs!"));
 					if (!Fence->IsSignaled())
 					{
@@ -416,7 +416,7 @@ bool FVulkanDynamicRHI::WaitAndProcessInterruptQueue()
 				// if one of these fences is already signaled, then that guarantees progress (no wait necessary on this pass)
 				if (NumFences == NumSyncs)
 				{
-					Device->GetFenceManager().WaitForAnyFence(MakeArrayView<VulkanRHI::FFence*>(Fences, NumFences), Timeout);
+					Device->GetFenceManager().WaitForAnyFence(MakeArrayView<FVulkanFence*>(Fences, NumFences), Timeout);
 				}
 			}
 			else
diff --git a/Engine/Source/Runtime/VulkanRHI/Private/VulkanSubmission.h b/Engine/Source/Runtime/VulkanRHI/Private/VulkanSubmission.h
--- a/Engine/Source/Runtime/VulkanRHI/Private/VulkanSubmission.h
+++ b/Engine/Source/Runtime/VulkanRHI/Private/VulkanSubmission.h
@@ -11,10 +11,12 @@
 
 class FVulkanDevice;
 class FVulkanDynamicRHI;
+class FVulkanFence;
 class FVulkanPendingGfxState;
 class FVulkanPendingComputeState;
 class FVulkanQueue;
 class FVulkanSwapChain;
+class FVulkanSemaphore;
 class FVulkanTiming;
 class FVulkanDescriptorPoolSetContainer;
 
@@ -63,16 +65,16 @@ protected:
 
 	FVulkanQueue& Queue;
 	TArray<VkPipelineStageFlags> WaitFlags;  // flags that match 1:1 with WaitSemaphores
-	TArray<VulkanRHI::FSemaphore*> WaitSemaphores; // wait before command buffers
+	TArray<FVulkanSemaphore*> WaitSemaphores; // wait before command buffers
 	TArray<FVulkanCommandBuffer*> CommandBuffers;
-	TArray<VulkanRHI::FSemaphore*> SignalSemaphores; // signaled after command buffers
+	TArray<FVulkanSemaphore*> SignalSemaphores; // signaled after command buffers
 
 	// Signaled when the payload has been submitted to the GPU queue
 	TArray<FGraphEventRef> SubmissionEvents;
 
 	// For internal completion tracking of the payload
 	uint64 TimelineSemaphoreValue = 0;
-	VulkanRHI::FFence* Fence = nullptr;
+	FVulkanFence* Fence = nullptr;
 
 	// Used to sync other CPU work to payload completion
 	TArray<FVulkanSyncPointRef> SyncPoints;
diff --git a/Engine/Source/Runtime/VulkanRHI/Private/VulkanSwapChain.cpp b/Engine/Source/Runtime/VulkanRHI/Private/VulkanSwapChain.cpp
--- a/Engine/Source/Runtime/VulkanRHI/Private/VulkanSwapChain.cpp
+++ b/Engine/Source/Runtime/VulkanRHI/Private/VulkanSwapChain.cpp
@@ -595,7 +595,7 @@ FVulkanSwapChain::FVulkanSwapChain(VkInstance InInstance, FVulkanDevice& InDevic
 
 #if VULKAN_USE_IMAGE_ACQUIRE_FENCES
 	ImageAcquiredFences.AddUninitialized(NumSwapChainImages);
-	VulkanRHI::FFenceManager& FenceMgr = Device.GetFenceManager();
+	FVulkanFenceManager& FenceMgr = Device.GetFenceManager();
 	for (uint32 BufferIndex = 0; BufferIndex < NumSwapChainImages; ++BufferIndex)
 	{
 		ImageAcquiredFences[BufferIndex] = Device.GetFenceManager().AllocateFence(true);
@@ -604,7 +604,7 @@ FVulkanSwapChain::FVulkanSwapChain(VkInstance InInstance, FVulkanDevice& InDevic
 	ImageAcquiredSemaphore.AddUninitialized(NumSwapChainImages);
 	for (uint32 BufferIndex = 0; BufferIndex < NumSwapChainImages; ++BufferIndex)
 	{
-		ImageAcquiredSemaphore[BufferIndex] = new VulkanRHI::FSemaphore(Device, VulkanRHI::EVulkanSemaphoreFlags::ExternallySignaled);
+		ImageAcquiredSemaphore[BufferIndex] = new FVulkanSemaphore(Device, EVulkanSemaphoreFlags::ExternallySignaled);
 		ImageAcquiredSemaphore[BufferIndex]->AddRef();
 	}
 
@@ -632,7 +632,7 @@ void FVulkanSwapChain::Destroy(FVulkanSwapChainRecreateInfo* RecreateInfo)
 	SwapChain = VK_NULL_HANDLE;
 
 #if VULKAN_USE_IMAGE_ACQUIRE_FENCES
-	VulkanRHI::FFenceManager& FenceMgr = Device.GetFenceManager();
+	FVulkanFenceManager& FenceMgr = Device.GetFenceManager();
 	for (int32 Index = 0; Index < ImageAcquiredFences.Num(); ++Index)
 	{
 		FenceMgr.ReleaseFence(ImageAcquiredFences[Index]);
@@ -653,7 +653,7 @@ void FVulkanSwapChain::Destroy(FVulkanSwapChainRecreateInfo* RecreateInfo)
 	Surface = VK_NULL_HANDLE;
 }
 
-int32 FVulkanSwapChain::AcquireImageIndex(VulkanRHI::FSemaphore** OutSemaphore)
+int32 FVulkanSwapChain::AcquireImageIndex(FVulkanSemaphore** OutSemaphore)
 {
 	check(CurrentImageIndex == -1);
 
@@ -667,7 +667,7 @@ int32 FVulkanSwapChain::AcquireImageIndex(VulkanRHI::FSemaphore** OutSemaphore)
 	// If we have not called present for any of the swapchain images, it will cause a crash/hang
 	checkf(!(NumAcquireCalls == ImageAcquiredSemaphore.Num() - 1 && NumPresentCalls == 0), TEXT("vkAcquireNextImageKHR will fail as no images have been presented before acquiring all of them"));
 #if VULKAN_USE_IMAGE_ACQUIRE_FENCES
-	VulkanRHI::FFenceManager& FenceMgr = Device.GetFenceManager();
+	FVulkanFenceManager& FenceMgr = Device.GetFenceManager();
 	FenceMgr.ResetFence(ImageAcquiredFences[SemaphoreIndex]);
 	const VkFence AcquiredFence = ImageAcquiredFences[SemaphoreIndex]->GetHandle();
 #else
@@ -789,7 +789,7 @@ void FVulkanSwapChain::RenderThreadPacing()
 	}
 }
 
-FVulkanSwapChain::EStatus FVulkanSwapChain::Present(FVulkanQueue* PresentQueue, VulkanRHI::FSemaphore* BackBufferRenderingDoneSemaphore)
+FVulkanSwapChain::EStatus FVulkanSwapChain::Present(FVulkanQueue* PresentQueue, FVulkanSemaphore* BackBufferRenderingDoneSemaphore)
 {
 	check(CurrentImageIndex != -1);
 
diff --git a/Engine/Source/Runtime/VulkanRHI/Private/VulkanSwapChain.h b/Engine/Source/Runtime/VulkanRHI/Private/VulkanSwapChain.h
--- a/Engine/Source/Runtime/VulkanRHI/Private/VulkanSwapChain.h
+++ b/Engine/Source/Runtime/VulkanRHI/Private/VulkanSwapChain.h
@@ -11,18 +11,14 @@
 #include "VulkanThirdParty.h"
 
 class FVulkanDevice;
+class FVulkanFence;
 class FVulkanQueue;
+class FVulkanSemaphore;
 class FVulkanTexture;
 class FVulkanView;
 enum EPixelFormat : uint8;
 
 
-namespace VulkanRHI
-{
-	class FFence;
-	class FSemaphore;
-}
-
 struct FVulkanSwapChainRecreateInfo
 {
 	VkSwapchainKHR SwapChain;
@@ -45,7 +41,7 @@ public:
 		OutOfDate = -1,
 		SurfaceLost = -2,
 	};
-	EStatus Present(FVulkanQueue* PresentQueue, VulkanRHI::FSemaphore* BackBufferRenderingDoneSemaphore);
+	EStatus Present(FVulkanQueue* PresentQueue, FVulkanSemaphore* BackBufferRenderingDoneSemaphore);
 
 	void RenderThreadPacing();
 	inline int8 DoesLockToVsync() { return LockToVsync; }
@@ -77,15 +73,15 @@ protected:
 	double NextPresentTargetTime = 0;
 
 	VkInstance Instance;
-	TArray<VulkanRHI::FSemaphore*> ImageAcquiredSemaphore;
+	TArray<FVulkanSemaphore*> ImageAcquiredSemaphore;
 #if VULKAN_USE_IMAGE_ACQUIRE_FENCES
-	TArray<VulkanRHI::FFence*> ImageAcquiredFences;
+	TArray<FVulkanFence*> ImageAcquiredFences;
 #endif
 	int8 LockToVsync;
 
 	uint32 PresentID = 0;
 
-	int32 AcquireImageIndex(VulkanRHI::FSemaphore** OutSemaphore);
+	int32 AcquireImageIndex(FVulkanSemaphore** OutSemaphore);
 
 
 	friend class FVulkanViewport;
diff --git a/Engine/Source/Runtime/VulkanRHI/Private/VulkanSynchronization.cpp b/Engine/Source/Runtime/VulkanRHI/Private/VulkanSynchronization.cpp
new file mode 100644
--- /dev/null
+++ b/Engine/Source/Runtime/VulkanRHI/Private/VulkanSynchronization.cpp
@@ -0,0 +1,258 @@
+// Copyright Epic Games, Inc. All Rights Reserved.
+
+#include "VulkanRHIPrivate.h"
+#include "VulkanSynchronization.h"
+#include "VulkanDevice.h"
+
+
+FVulkanFence::FVulkanFence(FVulkanDevice& InDevice, FVulkanFenceManager& InOwner, bool bCreateSignaled)
+	: State(bCreateSignaled ? FVulkanFence::EState::Signaled : FVulkanFence::EState::NotReady)
+	, Owner(InOwner)
+{
+	VkFenceCreateInfo Info;
+	ZeroVulkanStruct(Info, VK_STRUCTURE_TYPE_FENCE_CREATE_INFO);
+	Info.flags = bCreateSignaled ? VK_FENCE_CREATE_SIGNALED_BIT : 0;
+	VERIFYVULKANRESULT(VulkanRHI::vkCreateFence(InDevice.GetHandle(), &Info, VULKAN_CPU_ALLOCATOR, &Handle));
+}
+
+FVulkanFence::~FVulkanFence()
+{
+	checkf(Handle == VK_NULL_HANDLE, TEXT("Didn't get properly destroyed by FVulkanFenceManager!"));
+}
+
+
+
+FVulkanFenceManager::~FVulkanFenceManager()
+{
+	ensure(UsedFences.Num() == 0);
+}
+
+void FVulkanFenceManager::DestroyFence(FVulkanFence* Fence)
+{
+	// Does not need to go in the deferred deletion queue
+	VulkanRHI::vkDestroyFence(Device.GetHandle(), Fence->GetHandle(), VULKAN_CPU_ALLOCATOR);
+	Fence->Handle = VK_NULL_HANDLE;
+	delete Fence;
+}
+
+void FVulkanFenceManager::Deinit()
+{
+	FScopeLock Lock(&FenceLock);
+	ensureMsgf(UsedFences.Num() == 0, TEXT("No all fences are done!"));
+	for (FVulkanFence* Fence : FreeFences)
+	{
+		DestroyFence(Fence);
+	}
+}
+
+FVulkanFence* FVulkanFenceManager::AllocateFence(bool bCreateSignaled)
+{
+	FScopeLock Lock(&FenceLock);
+	if (FreeFences.Num() != 0)
+	{
+		FVulkanFence* Fence = FreeFences[0];
+		FreeFences.RemoveAtSwap(0, EAllowShrinking::No);
+		UsedFences.Add(Fence);
+
+		if (bCreateSignaled)
+		{
+			Fence->State = FVulkanFence::EState::Signaled;
+		}
+		return Fence;
+	}
+
+	FVulkanFence* NewFence = new FVulkanFence(Device, *this, bCreateSignaled);
+	UsedFences.Add(NewFence);
+	return NewFence;
+}
+
+// Sets it to nullptr
+void FVulkanFenceManager::ReleaseFence(FVulkanFence*& Fence)
+{
+	FScopeLock Lock(&FenceLock);
+	ResetFence(Fence);
+	UsedFences.RemoveSingleSwap(Fence, EAllowShrinking::No);
+#if VULKAN_REUSE_FENCES
+	FreeFences.Add(Fence);
+#else
+	DestroyFence(Fence);
+#endif
+	Fence = nullptr;
+}
+
+void FVulkanFenceManager::WaitAndReleaseFence(FVulkanFence*& Fence, uint64 TimeInNanoseconds)
+{
+	FScopeLock Lock(&FenceLock);
+	if (!Fence->IsSignaled())
+	{
+		WaitForFence(Fence, TimeInNanoseconds);
+	}
+
+	ResetFence(Fence);
+	UsedFences.RemoveSingleSwap(Fence, EAllowShrinking::No);
+	FreeFences.Add(Fence);
+	Fence = nullptr;
+}
+
+bool FVulkanFenceManager::CheckFenceState(FVulkanFence* Fence)
+{
+	check(UsedFences.Contains(Fence));
+	check(Fence->State == FVulkanFence::EState::NotReady);
+	const VkResult Result = VulkanRHI::vkGetFenceStatus(Device.GetHandle(), Fence->Handle);
+	switch (Result)
+	{
+	case VK_SUCCESS:
+		Fence->State = FVulkanFence::EState::Signaled;
+		return true;
+
+	case VK_NOT_READY:
+		break;
+
+	default:
+		VERIFYVULKANRESULT(Result);
+		break;
+	}
+
+	return false;
+}
+
+bool FVulkanFenceManager::WaitForFence(FVulkanFence* Fence, uint64 TimeInNanoseconds)
+{
+#if VULKAN_ENABLE_AGGRESSIVE_STATS
+	SCOPE_CYCLE_COUNTER(STAT_VulkanWaitFence);
+#endif
+
+	check(UsedFences.Contains(Fence));
+	check(Fence->State == FVulkanFence::EState::NotReady);
+	VkResult Result = VulkanRHI::vkWaitForFences(Device.GetHandle(), 1, &Fence->Handle, true, TimeInNanoseconds);
+	switch (Result)
+	{
+	case VK_SUCCESS:
+		Fence->State = FVulkanFence::EState::Signaled;
+		return true;
+	case VK_TIMEOUT:
+		break;
+	default:
+		VERIFYVULKANRESULT(Result);
+		break;
+	}
+
+	return false;
+}
+
+bool FVulkanFenceManager::WaitForAnyFence(TArrayView<FVulkanFence*> Fences, uint64 TimeInNanoseconds)
+{
+	// Mostly used for waits on queues, inlinealloc to that size
+	TArray<VkFence, TInlineAllocator<(int32)EVulkanQueueType::Count>> FenceHandles;
+	for (FVulkanFence* Fence : Fences)
+	{
+		if (Fence && !Fence->IsSignaled())
+		{
+			check(UsedFences.Contains(Fence));
+			FenceHandles.Add(Fence->GetHandle());
+		}
+	}
+
+	const VkResult Result = VulkanRHI::vkWaitForFences(Device.GetHandle(), FenceHandles.Num(), FenceHandles.GetData(), false, TimeInNanoseconds);
+	switch (Result)
+	{
+	case VK_SUCCESS:
+		if (Fences.Num() == 1)
+		{
+			Fences[0]->State = FVulkanFence::EState::Signaled;
+		}
+		return true;
+	case VK_TIMEOUT:
+		break;
+	default:
+		VERIFYVULKANRESULT(Result);
+		break;
+	}
+
+	return false;
+}
+
+void FVulkanFenceManager::ResetFence(FVulkanFence* Fence)
+{
+	if (Fence->State != FVulkanFence::EState::NotReady)
+	{
+		VERIFYVULKANRESULT(VulkanRHI::vkResetFences(Device.GetHandle(), 1, &Fence->Handle));
+		Fence->State = FVulkanFence::EState::NotReady;
+	}
+}
+
+
+
+
+FVulkanSemaphore::FVulkanSemaphore(FVulkanDevice& InDevice, EVulkanSemaphoreFlags InFlags, uint64 InInitialTimelineValue)
+	: Device(InDevice)
+	, SemaphoreHandle(VK_NULL_HANDLE)
+	, Flags(InFlags)
+{
+	// Create semaphore
+	VkSemaphoreCreateInfo CreateInfo;
+	ZeroVulkanStruct(CreateInfo, VK_STRUCTURE_TYPE_SEMAPHORE_CREATE_INFO);
+
+	VkSemaphoreTypeCreateInfo TypeCreateInfo;
+	if (IsTimeline())
+	{
+		checkf(InDevice.GetOptionalExtensions().HasKHRTimelineSemaphore, TEXT("Trying to create timeline semaphore without device support!"));
+		ZeroVulkanStruct(TypeCreateInfo, VK_STRUCTURE_TYPE_SEMAPHORE_TYPE_CREATE_INFO);
+		TypeCreateInfo.semaphoreType = VK_SEMAPHORE_TYPE_TIMELINE;
+		TypeCreateInfo.initialValue = (uint64_t)InInitialTimelineValue;
+		CreateInfo.pNext = &TypeCreateInfo;
+	}
+
+	VERIFYVULKANRESULT(VulkanRHI::vkCreateSemaphore(Device.GetHandle(), &CreateInfo, VULKAN_CPU_ALLOCATOR, &SemaphoreHandle));
+}
+
+FVulkanSemaphore::FVulkanSemaphore(FVulkanDevice& InDevice, const VkSemaphore& InExternalSemaphore)
+	: Device(InDevice)
+	, SemaphoreHandle(InExternalSemaphore)
+	, Flags(EVulkanSemaphoreFlags::ExternallySignaled)
+{}
+
+FVulkanSemaphore::~FVulkanSemaphore()
+{
+	check(SemaphoreHandle != VK_NULL_HANDLE);
+	if (EnumHasAnyFlags(Flags, EVulkanSemaphoreFlags::ImmediateDeletion))
+	{
+		VulkanRHI::vkDestroySemaphore(Device.GetHandle(), SemaphoreHandle, VULKAN_CPU_ALLOCATOR);
+	}
+	else if (!IsExternallyOwned())
+	{
+		Device.GetDeferredDeletionQueue().EnqueueResource(VulkanRHI::FDeferredDeletionQueue2::EType::Semaphore, SemaphoreHandle);
+	}
+	SemaphoreHandle = VK_NULL_HANDLE;
+}
+
+bool FVulkanSemaphore::WaitForTimelineSemaphoreValue(uint64 Value, uint64 Timeout)
+{
+	checkf(IsTimeline(), TEXT("WaitForTimelineSemaphoreValue not available!  Semaphore was not created with EVulkanSemaphoreFlags::Timeline."));
+	checkSlow(SemaphoreHandle != VK_NULL_HANDLE);
+	VkSemaphoreWaitInfo WaitInfo;
+	ZeroVulkanStruct(WaitInfo, VK_STRUCTURE_TYPE_SEMAPHORE_WAIT_INFO);
+	WaitInfo.semaphoreCount = 1;
+	WaitInfo.pSemaphores = &SemaphoreHandle;
+	WaitInfo.pValues = (uint64_t*)&Value;
+	const VkResult Result = VulkanRHI::vkWaitSemaphoresKHR(Device.GetHandle(), &WaitInfo, Timeout);
+	switch (Result)
+	{
+	case VK_SUCCESS:
+		return true;
+	case VK_TIMEOUT:
+		break;
+	default:
+		VERIFYVULKANRESULT(Result);
+		break;
+	}
+	return false;
+}
+
+uint64 FVulkanSemaphore::GetTimelineSemaphoreValue()
+{
+	checkf(IsTimeline(), TEXT("GetTimelineSemaphoreValue not available!  Semaphore was not created with EVulkanSemaphoreFlags::Timeline."));
+	uint64 Value = 0;
+	VERIFYVULKANRESULT(VulkanRHI::vkGetSemaphoreCounterValueKHR(Device.GetHandle(), SemaphoreHandle, (uint64_t*)&Value));
+	return Value;
+}
diff --git a/Engine/Source/Runtime/VulkanRHI/Private/VulkanSynchronization.h b/Engine/Source/Runtime/VulkanRHI/Private/VulkanSynchronization.h
new file mode 100644
--- /dev/null
+++ b/Engine/Source/Runtime/VulkanRHI/Private/VulkanSynchronization.h
@@ -0,0 +1,161 @@
+// Copyright Epic Games, Inc. All Rights Reserved.
+
+#pragma once 
+
+#include "Containers/Array.h"
+#include "HAL/CriticalSection.h"
+#include "Templates/RefCounting.h"
+#include "VulkanConfiguration.h"
+#include "VulkanThirdParty.h"
+
+class FVulkanDevice;
+class FVulkanFenceManager;
+
+// Internal Vulkan Fence class that wrap VkFence management
+// (not to be confused with FVulkanGPUFence class that maps to RHI's FRHIFence) 
+class FVulkanFence
+{
+public:
+	FVulkanFence(FVulkanDevice& InDevice, FVulkanFenceManager& InOwner, bool bCreateSignaled);
+
+	inline VkFence GetHandle() const
+	{
+		return Handle;
+	}
+
+	inline bool IsSignaled() const
+	{
+		return State == EState::Signaled;
+	}
+
+	FVulkanFenceManager& GetOwner()
+	{
+		return Owner;
+	}
+
+protected:
+	VkFence Handle;
+
+	enum class EState
+	{
+		// Initial state
+		NotReady,
+
+		// After GPU processed it
+		Signaled,
+	};
+
+	EState State;
+
+	FVulkanFenceManager& Owner;
+
+	// Only owner can delete!
+	~FVulkanFence();
+	friend FVulkanFenceManager;
+};
+
+class FVulkanFenceManager
+{
+public:
+	FVulkanFenceManager(FVulkanDevice& InDevice)
+		: Device(InDevice)
+	{
+	}
+	~FVulkanFenceManager();
+
+	void Deinit();
+
+	FVulkanFence* AllocateFence(bool bCreateSignaled = false);
+
+	inline bool IsFenceSignaled(FVulkanFence* Fence)
+	{
+		if (Fence->IsSignaled())
+		{
+			return true;
+		}
+
+		return CheckFenceState(Fence);
+	}
+
+	// Returns false if it timed out
+	bool WaitForAnyFence(TArrayView<FVulkanFence*> Fences, uint64 TimeInNanoseconds);
+	bool WaitForFence(FVulkanFence* Fence, uint64 TimeInNanoseconds);
+
+	void ResetFence(FVulkanFence* Fence);
+
+	// Sets it to nullptr
+	void ReleaseFence(FVulkanFence*& Fence);
+
+	// Sets it to nullptr
+	void WaitAndReleaseFence(FVulkanFence*& Fence, uint64 TimeInNanoseconds);
+
+protected:
+	FVulkanDevice& Device;
+	FCriticalSection FenceLock;
+	TArray<FVulkanFence*> FreeFences;
+	TArray<FVulkanFence*> UsedFences;
+
+	// Returns true if signaled
+	bool CheckFenceState(FVulkanFence* Fence);
+
+	void DestroyFence(FVulkanFence* Fence);
+};
+
+
+
+enum class EVulkanSemaphoreFlags : uint8
+{
+	None = 0,
+
+	// Will not delete handle on destruction
+	ExternallyOwned    = 1 << 1,    
+
+	// Inform submission pipeline that the signal will not come from a payload (acquired image semaphore for example)
+	ExternallySignaled = 1 << 2,
+
+	// Create a timeline semaphore (must be supported)
+	Timeline           = 1 << 3,
+
+	// Will not be queued for deferred deletion
+	ImmediateDeletion  = 1 << 4
+};
+ENUM_CLASS_FLAGS(EVulkanSemaphoreFlags);
+
+
+// Internal Vulkan Semaphore class that wrap VkSemaphore management (binary or timeline)
+class FVulkanSemaphore : public FThreadSafeRefCountedObject
+{
+public:
+	FVulkanSemaphore(FVulkanDevice& InDevice, EVulkanSemaphoreFlags InFlags=EVulkanSemaphoreFlags::None, uint64 InInitialTimelineValue=0);
+	FVulkanSemaphore(FVulkanDevice& InDevice, const VkSemaphore& InExternalSemaphore);
+	virtual ~FVulkanSemaphore();
+
+	VkSemaphore GetHandle() const
+	{
+		return SemaphoreHandle;
+	}
+
+	bool IsExternallyOwned() const
+	{
+		return EnumHasAnyFlags(Flags, EVulkanSemaphoreFlags::ExternallyOwned);
+	}
+
+	bool IsExternallySignaled() const
+	{
+		return EnumHasAnyFlags(Flags, EVulkanSemaphoreFlags::ExternallySignaled);
+	}
+
+	bool IsTimeline() const
+	{
+		return EnumHasAnyFlags(Flags, EVulkanSemaphoreFlags::Timeline);
+	}
+
+	uint64 GetTimelineSemaphoreValue();
+	bool WaitForTimelineSemaphoreValue(uint64 Value, uint64 Timeout);
+
+private:
+	FVulkanDevice& Device;
+	VkSemaphore SemaphoreHandle;
+	const EVulkanSemaphoreFlags Flags;
+};
+
diff --git a/Engine/Source/Runtime/VulkanRHI/Private/VulkanViewport.cpp b/Engine/Source/Runtime/VulkanRHI/Private/VulkanViewport.cpp
--- a/Engine/Source/Runtime/VulkanRHI/Private/VulkanViewport.cpp
+++ b/Engine/Source/Runtime/VulkanRHI/Private/VulkanViewport.cpp
@@ -442,7 +442,7 @@ void FVulkanViewport::CreateSwapchain(FVulkanCommandListContext& Context, FVulka
 			{
 				for (int32 Index = 0, NumBuffers = RenderingDoneSemaphores.Num(); Index < NumBuffers; ++Index)
 				{
-					RenderingDoneSemaphores[Index] = new VulkanRHI::FSemaphore(Device);
+					RenderingDoneSemaphores[Index] = new FVulkanSemaphore(Device);
 					RenderingDoneSemaphores[Index]->AddRef();
 				}
 			}
diff --git a/Engine/Source/Runtime/VulkanRHI/Private/VulkanViewport.h b/Engine/Source/Runtime/VulkanRHI/Private/VulkanViewport.h
--- a/Engine/Source/Runtime/VulkanRHI/Private/VulkanViewport.h
+++ b/Engine/Source/Runtime/VulkanRHI/Private/VulkanViewport.h
@@ -13,6 +13,7 @@ class FVulkanDynamicRHI;
 class FVulkanSwapChain;
 class FVulkanQueue;
 class FVulkanViewport;
+class FVulkanSemaphore;
 
 namespace VulkanRHI
 {
@@ -116,7 +117,7 @@ protected:
 
 	// NUM_BUFFERS don't have to match exactly as the driver can require a minimum number larger than NUM_BUFFERS. Provide some slack
 	TArray<TRefCountPtr<FVulkanTexture>, TInlineAllocator<NUM_BUFFERS*2>> BackBufferImages;
-	TArray<VulkanRHI::FSemaphore*, TInlineAllocator<NUM_BUFFERS*2>> RenderingDoneSemaphores;
+	TArray<FVulkanSemaphore*, TInlineAllocator<NUM_BUFFERS*2>> RenderingDoneSemaphores;
 	TIndirectArray<FVulkanView, TInlineAllocator<NUM_BUFFERS*2>> TextureViews;
 	TRefCountPtr<FVulkanBackBuffer> RHIBackBuffer;
 
@@ -139,7 +140,7 @@ protected:
 	int8 LockToVsync;
 
 	// Just a pointer, not owned by this class
-	VulkanRHI::FSemaphore* AcquiredSemaphore;
+	FVulkanSemaphore* AcquiredSemaphore;
 
 	FCustomPresentRHIRef CustomPresent;
 
