VulkanRHI: Fix reserved resources getting destroyed with pending commits in flight

#rnx

diff --git a/Engine/Source/Runtime/VulkanRHI/Private/VulkanQueue.cpp b/Engine/Source/Runtime/VulkanRHI/Private/VulkanQueue.cpp
--- a/Engine/Source/Runtime/VulkanRHI/Private/VulkanQueue.cpp
+++ b/Engine/Source/Runtime/VulkanRHI/Private/VulkanQueue.cpp
@@ -105,7 +105,7 @@ void FVulkanQueue::BindSparseResources(FVulkanPayload& Payload)
 	{
 		if (CommitDesc.Resource->GetType() == RRT_Buffer)
 		{
-			FVulkanBuffer* Buffer = ResourceCast(static_cast<FRHIBuffer*>(CommitDesc.Resource));
+			FVulkanBuffer* Buffer = ResourceCast(static_cast<FRHIBuffer*>(CommitDesc.Resource.GetReference()));
 			TArray<VkSparseMemoryBind> BufferSparseMemoryBinds = Buffer->CommitReservedResource(CommitDesc.CommitSizeInBytes);
 			if (BufferSparseMemoryBinds.Num())
 			{
@@ -120,7 +120,7 @@ void FVulkanQueue::BindSparseResources(FVulkanPayload& Payload)
 		}
 		else if (CommitDesc.Resource->GetType() == RRT_Texture)
 		{
-			FVulkanTexture* Texture = ResourceCast(static_cast<FRHITexture*>(CommitDesc.Resource));
+			FVulkanTexture* Texture = ResourceCast(static_cast<FRHITexture*>(CommitDesc.Resource.GetReference()));
 			TArray<VkSparseMemoryBind> ImageSparseMemoryBinds = Texture->CommitReservedResource(CommitDesc.CommitSizeInBytes);
 			if (ImageSparseMemoryBinds.Num())
 			{
diff --git a/Engine/Source/Runtime/VulkanRHI/Private/VulkanSubmission.h b/Engine/Source/Runtime/VulkanRHI/Private/VulkanSubmission.h
--- a/Engine/Source/Runtime/VulkanRHI/Private/VulkanSubmission.h
+++ b/Engine/Source/Runtime/VulkanRHI/Private/VulkanSubmission.h
@@ -33,7 +33,7 @@ static FVulkanSyncPointRef CreateVulkanSyncPoint()
 
 struct FVulkanCommitReservedResourceDesc
 {
-	FRHIResource* Resource = nullptr;
+	TRefCountPtr<FRHIResource> Resource = nullptr;
 	uint64 CommitSizeInBytes = 0;
 };
 
