VulkanRHI: Use a rename (alloc swap) instead of a staging buffer if the full buffer is locked on unified memory device

#rnx

diff --git a/Engine/Source/Runtime/VulkanRHI/Private/VulkanBuffer.cpp b/Engine/Source/Runtime/VulkanRHI/Private/VulkanBuffer.cpp
--- a/Engine/Source/Runtime/VulkanRHI/Private/VulkanBuffer.cpp
+++ b/Engine/Source/Runtime/VulkanRHI/Private/VulkanBuffer.cpp
@@ -400,6 +400,10 @@ void* FVulkanBuffer::Lock(FRHICommandListBase& RHICmdList, EResourceLockMode Loc
 	{
 		check((LockMode == RLM_WriteOnly) || (LockMode == RLM_WriteOnly_NoOverwrite));
 
+		// When the entire buffer is locked on unified memory devices, 
+		// we can use a new allocation instead of staging+copy (at no extra cost since all memory is host visible)
+		const bool bEntireBuffer = bUnifiedMem && (LockSize == GetSize()) && (Offset == 0);
+
 		// If this is the first lock on host visible memory, then the memory is still untouched so use it directly
 		const bool bIsHostVisible = (bUnifiedMem || bDynamic);
 		if (bIsHostVisible && (bIsFirstLock || (LockMode == RLM_WriteOnly_NoOverwrite)))
@@ -409,7 +413,7 @@ void* FVulkanBuffer::Lock(FRHICommandListBase& RHICmdList, EResourceLockMode Loc
 			DataOffset = Offset;
 			LockStatus = ELockStatus::PersistentMapping;
 		}
-		else if (bStatic || GVulkanForceStagingBufferOnLock || bIsSparseResource)
+		else if ((bStatic && !bEntireBuffer) || GVulkanForceStagingBufferOnLock || bIsSparseResource)
 		{
 			FVulkanPendingBufferLock PendingLock;
 			PendingLock.Offset = Offset;
