PR #13426: Avoid usage of VK_KHR_separate_depth_stencil_layouts in renderpass setup, on mali

diff --git a/Engine/Source/Runtime/VulkanRHI/Private/Android/VulkanAndroidPlatform.cpp b/Engine/Source/Runtime/VulkanRHI/Private/Android/VulkanAndroidPlatform.cpp
--- a/Engine/Source/Runtime/VulkanRHI/Private/Android/VulkanAndroidPlatform.cpp
+++ b/Engine/Source/Runtime/VulkanRHI/Private/Android/VulkanAndroidPlatform.cpp
@@ -118,6 +118,7 @@ TArray<TArray<ANSICHAR>> FVulkanAndroidPlatform::SwappyRequiredExtensions;
 int32 FVulkanAndroidPlatform::AFBCWorkaroundOption = 0;
 int32 FVulkanAndroidPlatform::ASTCWorkaroundOption = 0;
 bool FVulkanAndroidPlatform::bRequiresDepthStencilFullWrite = false;
+bool FVulkanAndroidPlatform::bRequiresMergedRenderPassDepthStencilLayout = false;
 
 #define CHECK_VK_ENTRYPOINTS(Type,Func) if (VulkanDynamicAPI::Func == NULL) { bFoundAllEntryPoints = false; UE_LOG(LogRHI, Warning, TEXT("Failed to find entry point for %s"), TEXT(#Func)); }
 
@@ -874,6 +875,20 @@ void FVulkanAndroidPlatform::SetupRequiresDepthStencilFullWriteWorkaround(const
 	}
 }
 
+void FVulkanAndroidPlatform::SetupRequiresMergedDepthStencilLayoutWorkaround(const FVulkanDevice& Device)
+{
+	bRequiresMergedRenderPassDepthStencilLayout = false;
+	if (Device.GetVendorId() == EGpuVendorId::Arm)
+	{
+		const uint32_t DriverVersion = Device.GetDeviceProperties().driverVersion;
+		if (DriverVersion < VK_MAKE_API_VERSION(0, 54, 1, 0))
+		{
+			bRequiresMergedRenderPassDepthStencilLayout = true;
+			UE_LOG(LogRHI, Display, TEXT("Enabling workaround that avoids usage of dedicated stencil layout via VkAttachmentDescriptionStencilLayout."));
+		}
+	}
+}
+
 bool FVulkanAndroidPlatform::FramePace(FVulkanDevice& Device, void* WindowHandle, VkSwapchainKHR Swapchain, uint32 PresentID, VkPresentInfoKHR& Info)
 {
 	QUICK_SCOPE_CYCLE_COUNTER(STAT_VulkanAndroid_FramePace);
@@ -1975,6 +1990,7 @@ void FVulkanAndroidPlatform::PostInitGPU(const FVulkanDevice& InDevice)
 {
 	SetupImageMemoryRequirementWorkaround(InDevice);
 	SetupRequiresDepthStencilFullWriteWorkaround(InDevice);
+	SetupRequiresMergedDepthStencilLayoutWorkaround(InDevice);
 
 	// start external compilers if precaching is specified.	
 	static const auto CVarChunkedPSOCache = IConsoleManager::Get().FindConsoleVariable(TEXT("r.Vulkan.UseChunkedPSOCache"));
diff --git a/Engine/Source/Runtime/VulkanRHI/Private/Android/VulkanAndroidPlatform.h b/Engine/Source/Runtime/VulkanRHI/Private/Android/VulkanAndroidPlatform.h
--- a/Engine/Source/Runtime/VulkanRHI/Private/Android/VulkanAndroidPlatform.h
+++ b/Engine/Source/Runtime/VulkanRHI/Private/Android/VulkanAndroidPlatform.h
@@ -209,6 +209,10 @@ public:
 	static bool RequiresDepthStencilFullWrite() { return bRequiresDepthStencilFullWrite; }
 	static void SetupRequiresDepthStencilFullWriteWorkaround(const FVulkanDevice& Device);
 
+	// Does the platform require to avoid VkAttachmentDescriptionStencilLayout in favor of merged depth-stencil layouts
+	static bool RequiresMergedDepthStencilLayout() { return bRequiresMergedRenderPassDepthStencilLayout; }
+	static void SetupRequiresMergedDepthStencilLayoutWorkaround(const FVulkanDevice& Device);
+
 	static bool FramePace(FVulkanDevice& Device, void* WindowHandle, VkSwapchainKHR Swapchain, uint32 PresentID, VkPresentInfoKHR& Info);
 
 	static VkResult Present(VkQueue Queue, VkPresentInfoKHR& PresentInfo);
@@ -283,6 +287,7 @@ protected:
 	static int32 ASTCWorkaroundOption;
 
 	static bool bRequiresDepthStencilFullWrite;
+	static bool bRequiresMergedRenderPassDepthStencilLayout;
 };
 
 #if VULKAN_SUPPORTS_GOOGLE_DISPLAY_TIMING
diff --git a/Engine/Source/Runtime/VulkanRHI/Private/VulkanGenericPlatform.h b/Engine/Source/Runtime/VulkanRHI/Private/VulkanGenericPlatform.h
--- a/Engine/Source/Runtime/VulkanRHI/Private/VulkanGenericPlatform.h
+++ b/Engine/Source/Runtime/VulkanRHI/Private/VulkanGenericPlatform.h
@@ -129,6 +129,9 @@ public:
 	// Does the platform require depth to be written on stencil clear
 	static bool RequiresDepthStencilFullWrite() { return false; }
 
+	// Does the platform require to avoid VkAttachmentDescriptionStencilLayout in favor of merged depth-stencil layouts
+	static bool RequiresMergedDepthStencilLayout() { return false; }
+
 	// Checks if the PSO cache matches the expected vulkan device properties
 	static bool PSOBinaryCacheMatches(FVulkanDevice* Device, const TArray<uint8>& DeviceCache);
 
diff --git a/Engine/Source/Runtime/VulkanRHI/Private/VulkanRenderpass.h b/Engine/Source/Runtime/VulkanRHI/Private/VulkanRenderpass.h
--- a/Engine/Source/Runtime/VulkanRHI/Private/VulkanRenderpass.h
+++ b/Engine/Source/Runtime/VulkanRHI/Private/VulkanRenderpass.h
@@ -76,6 +76,10 @@ template <>
 inline void FVulkanAttachmentReference<VkAttachmentReference2>::SetDepthStencilAttachment(const VkAttachmentReference& AttachmentReferenceIn, 
 	const VkAttachmentReferenceStencilLayout* StencilReference, VkImageAspectFlags AspectMask, bool bSupportsParallelRendering)
 {
+	if (FVulkanPlatform::RequiresMergedDepthStencilLayout())
+	{
+		bSupportsParallelRendering = false;
+	}
 	sType = VK_STRUCTURE_TYPE_ATTACHMENT_REFERENCE_2;
 	pNext = (bSupportsParallelRendering && StencilReference && StencilReference->stencilLayout != VK_IMAGE_LAYOUT_UNDEFINED) ? StencilReference : nullptr;
 	attachment = AttachmentReferenceIn.attachment;
@@ -319,6 +323,10 @@ struct FVulkanAttachmentDescription<VkAttachmentDescription2>
 
 	FVulkanAttachmentDescription(const VkAttachmentDescription& InDesc, const VkAttachmentDescriptionStencilLayout* InStencilDesc, bool bSupportsParallelRendering)
 	{
+		if (FVulkanPlatform::RequiresMergedDepthStencilLayout())
+		{
+			bSupportsParallelRendering = false;
+		}
 		const bool bHasStencilLayout = bSupportsParallelRendering && VulkanRHI::VulkanFormatHasStencil(InDesc.format) && (InStencilDesc != nullptr);
 
 		sType = VK_STRUCTURE_TYPE_ATTACHMENT_DESCRIPTION_2;
