PR #13672: Fix Vulkan Queries During Multiview

#rnx

diff --git a/Engine/Source/Runtime/VulkanRHI/Private/VulkanCommandBuffer.cpp b/Engine/Source/Runtime/VulkanRHI/Private/VulkanCommandBuffer.cpp
--- a/Engine/Source/Runtime/VulkanRHI/Private/VulkanCommandBuffer.cpp
+++ b/Engine/Source/Runtime/VulkanRHI/Private/VulkanCommandBuffer.cpp
@@ -91,6 +91,7 @@ void FVulkanCommandBuffer::EndRenderPass()
 	checkf(IsInsideRenderPass(), TEXT("Can't EndRP as we're NOT inside one! CmdBuffer 0x%p State=%d"), CommandBufferHandle, (int32)State);
 	VulkanRHI::vkCmdEndRenderPass(CommandBufferHandle);
 	State = EState::IsInsideBegin;
+	CurrentMultiViewCount = 0;
 }
 
 void FVulkanCommandBuffer::BeginRenderPass(const FVulkanBeginRenderPassInfo& BeginRenderPassInfo, const VkClearValue* AttachmentClearValues)
@@ -122,6 +123,7 @@ void FVulkanCommandBuffer::BeginRenderPass(const FVulkanBeginRenderPassInfo& Beg
 		VulkanRHI::vkCmdBeginRenderPass(CommandBufferHandle, &Info, SubpassContents);
 	}
 
+	CurrentMultiViewCount = BeginRenderPassInfo.RenderPass.GetLayout().GetMultiViewCount();
 	State = EState::IsInsideRenderPass;
 }
 
diff --git a/Engine/Source/Runtime/VulkanRHI/Private/VulkanCommandBuffer.h b/Engine/Source/Runtime/VulkanRHI/Private/VulkanCommandBuffer.h
--- a/Engine/Source/Runtime/VulkanRHI/Private/VulkanCommandBuffer.h
+++ b/Engine/Source/Runtime/VulkanRHI/Private/VulkanCommandBuffer.h
@@ -115,6 +115,8 @@ public:
 	uint8 bHasScissor					: 1 = 0;
 	uint8 bHasStencilRef				: 1 = 0;
 
+	uint8 CurrentMultiViewCount = 0;
+
 	// Every secondary command buffer executed from this one with tied lifetimes
 	TArray<FVulkanCommandBuffer*> ExecutedSecondaryCommandBuffers;
 
diff --git a/Engine/Source/Runtime/VulkanRHI/Private/VulkanContext.h b/Engine/Source/Runtime/VulkanRHI/Private/VulkanContext.h
--- a/Engine/Source/Runtime/VulkanRHI/Private/VulkanContext.h
+++ b/Engine/Source/Runtime/VulkanRHI/Private/VulkanContext.h
@@ -432,7 +432,7 @@ protected:
 	const ERHIPipeline RHIPipeline;
 	const bool bSupportsBreadcrumbs;
 
-	void BeginOcclusionQueryBatch(uint32 NumQueriesInBatch);
+	void BeginOcclusionQueryBatch(uint32 NumQueriesInBatch, uint32 MultiViewCount);
 
 	TArray<FString> EventStack;
 
diff --git a/Engine/Source/Runtime/VulkanRHI/Private/VulkanQuery.cpp b/Engine/Source/Runtime/VulkanRHI/Private/VulkanQuery.cpp
--- a/Engine/Source/Runtime/VulkanRHI/Private/VulkanQuery.cpp
+++ b/Engine/Source/Runtime/VulkanRHI/Private/VulkanQuery.cpp
@@ -68,29 +68,32 @@ FVulkanQueryPool::~FVulkanQueryPool()
 void FVulkanQueryPool::ReserveQuery(FVulkanRenderQuery* Query)
 {
 	checkSlow(QueryType == Query->GetQueryPoolType());
-	Query->IndexInPool = CurrentQueryCount++;
-	QueryRefs[Query->IndexInPool] = Query;
-	QueryResults[Query->IndexInPool] = &Query->Result;
+	const uint32 QuerySlot = CurrentQueryCount++;
+	Query->IndexInPool = QuerySlot * QueryStride;
+	QueryRefs[QuerySlot] = Query;
+	QueryResults[QuerySlot] = &Query->Result;
 }
 
 uint32 FVulkanQueryPool::ReserveQuery(uint64* ResultPtr)
 {
 	checkSlow(QueryType == EVulkanQueryPoolType::Timestamp);
-	const uint32 IndexInPool = CurrentQueryCount++;
-	QueryRefs[IndexInPool] = nullptr;
-	QueryResults[IndexInPool] = ResultPtr;
+	const uint32 QuerySlot = CurrentQueryCount++;
+	const uint32 IndexInPool = QuerySlot * QueryStride;
+	QueryRefs[QuerySlot] = nullptr;
+	QueryResults[QuerySlot] = ResultPtr;
 	return IndexInPool;
 }
 
-void FVulkanQueryPool::Reset(FVulkanCommandBuffer& InCmdBuffer)
+void FVulkanQueryPool::Reset(FVulkanCommandBuffer& InCmdBuffer, uint32 InQueryStride)
 {
 	CurrentQueryCount = 0;
+	QueryStride = InQueryStride;
 	UnusedFrameCount = 0;
 
 	QueryRefs.Empty(MaxQueries);
 	QueryResults.Empty(MaxQueries);
-	QueryRefs.AddDefaulted(MaxQueries);
-	QueryResults.SetNumZeroed(MaxQueries);
+	QueryRefs.AddDefaulted(MaxQueries / QueryStride);
+	QueryResults.SetNumZeroed(MaxQueries / QueryStride);
 
 	if (Device.GetOptionalExtensions().HasEXTHostQueryReset)
 	{
@@ -107,11 +110,12 @@ bool FVulkanQueryPool::IsStale() const
 	return UnusedFrameCount >= GQueryPoolDeletionDelay;
 }
 
-void FVulkanCommandListContext::BeginOcclusionQueryBatch(uint32 NumQueriesInBatch)
+void FVulkanCommandListContext::BeginOcclusionQueryBatch(uint32 NumQueriesInBatch, uint32 MultiViewCount)
 {
+	uint32 QueryStride = MultiViewCount == 0 ? 1 : MultiViewCount;
 	TArray<FVulkanQueryPool*>& OcclusionPoolArray = GetQueryPoolArray(EVulkanQueryPoolType::Occlusion);
-	FVulkanQueryPool* NewOcclusionQueryPool = Device.AcquireOcclusionQueryPool(NumQueriesInBatch);
-	NewOcclusionQueryPool->Reset(GetCommandBuffer());
+	FVulkanQueryPool* NewOcclusionQueryPool = Device.AcquireOcclusionQueryPool(NumQueriesInBatch * QueryStride);
+	NewOcclusionQueryPool->Reset(GetCommandBuffer(), QueryStride);
 	OcclusionPoolArray.Add(NewOcclusionQueryPool);
 }
 
@@ -369,7 +373,16 @@ FVulkanQueryPool* FVulkanCommandListContext::GetCurrentOcclusionQueryPool()
 FVulkanQueryPool* FVulkanContextCommon::GetCurrentTimestampQueryPool(FVulkanPayload& Payload)
 {
 	TArray<FVulkanQueryPool*>& TimestampPoolArray = Payload.QueryPools[(int32)EVulkanQueryPoolType::Timestamp];
-	if ((TimestampPoolArray.Num() == 0) || TimestampPoolArray.Last()->IsFull())
+	uint32 DesiredQueryStride = 1u;
+	if (Payload.CommandBuffers.Num() > 0 && Payload.CommandBuffers.Last()->IsInsideRenderPass())
+	{
+		DesiredQueryStride = Payload.CommandBuffers.Last()->CurrentMultiViewCount;
+		DesiredQueryStride = DesiredQueryStride == 0 ? 1 : DesiredQueryStride;
+	}
+	// Note: In theory, it's possible to have both multiview and non-multiview timing queries in one pool.
+	// However, this would add significant complexity to fetching results.
+	// For now, to keep things simple, switch to a new query pool when the stride changes.
+	if ((TimestampPoolArray.Num() == 0) || TimestampPoolArray.Last()->IsFull() || TimestampPoolArray.Last()->GetQueryStride() != DesiredQueryStride)
 	{
 		FVulkanQueryPool* NewPool = Device.AcquireTimingQueryPool();
 
@@ -377,7 +390,7 @@ FVulkanQueryPool* FVulkanContextCommon::GetCurrentTimestampQueryPool(FVulkanPayl
 		{
 			PrepareNewCommandBuffer(Payload);
 		}
-		NewPool->Reset(*Payload.CommandBuffers.Last());
+		NewPool->Reset(*Payload.CommandBuffers.Last(), DesiredQueryStride);
 
 		TimestampPoolArray.Add(NewPool);
 	}
diff --git a/Engine/Source/Runtime/VulkanRHI/Private/VulkanQuery.h b/Engine/Source/Runtime/VulkanRHI/Private/VulkanQuery.h
--- a/Engine/Source/Runtime/VulkanRHI/Private/VulkanQuery.h
+++ b/Engine/Source/Runtime/VulkanRHI/Private/VulkanQuery.h
@@ -71,15 +71,20 @@ public:
 		return Device;
 	}
 
+	uint32 GetQueryStride() const
+	{
+		return QueryStride;
+	}
+
 	bool IsFull() const
 	{
-		return (CurrentQueryCount >= MaxQueries);
+		return (CurrentQueryCount * QueryStride >= MaxQueries);
 	}
 
 	void ReserveQuery(FVulkanRenderQuery* Query);
 	uint32 ReserveQuery(uint64* ResultPtr);
 
-	void Reset(FVulkanCommandBuffer& InCmdBuffer);
+	void Reset(FVulkanCommandBuffer& InCmdBuffer, uint32 InQueryStride = 1u);
 	bool IsStale() const;
 
 	void IncrementUnusedFrameCount()
@@ -95,6 +100,11 @@ protected:
 	TArray<TRefCountPtr<FVulkanRenderQuery>> QueryRefs;
 	TArray<uint64*> QueryResults;
 	uint32 CurrentQueryCount = 0;
+	/** When queries are issued during a multiview render pass, vkCmdBeginQuery and vkCmdEndQuery
+	 *  each correspond to multiple subsequent queries in the pool, one for each view.
+	 *  The interpretation of these per-view values is implementation dependent, but they
+	 *  can be summed together to get the actual value. */
+	uint32 QueryStride = 1u;
 	int32 UnusedFrameCount = 0;
 
 	friend FVulkanDynamicRHI;
diff --git a/Engine/Source/Runtime/VulkanRHI/Private/VulkanRenderTarget.cpp b/Engine/Source/Runtime/VulkanRHI/Private/VulkanRenderTarget.cpp
--- a/Engine/Source/Runtime/VulkanRHI/Private/VulkanRenderTarget.cpp
+++ b/Engine/Source/Runtime/VulkanRHI/Private/VulkanRenderTarget.cpp
@@ -560,7 +560,7 @@ void FVulkanCommandListContext::RHIBeginRenderPass(const FRHIRenderPassInfo& InI
 
 	if (InInfo.NumOcclusionQueries > 0)
 	{
-		BeginOcclusionQueryBatch(InInfo.NumOcclusionQueries);
+		BeginOcclusionQueryBatch(InInfo.NumOcclusionQueries, InInfo.MultiViewCount);
 	}
 
 	// Acquire a descriptor pool set on a first render pass
diff --git a/Engine/Source/Runtime/VulkanRHI/Private/VulkanSubmission.cpp b/Engine/Source/Runtime/VulkanRHI/Private/VulkanSubmission.cpp
--- a/Engine/Source/Runtime/VulkanRHI/Private/VulkanSubmission.cpp
+++ b/Engine/Source/Runtime/VulkanRHI/Private/VulkanSubmission.cpp
@@ -664,7 +664,7 @@ void FVulkanDynamicRHI::CompletePayload(FVulkanPayload* Payload)
 					Device->GetHandle(),
 					QueryPool->GetHandle(),
 					0,
-					QueryPool->CurrentQueryCount,
+					QueryPool->CurrentQueryCount * QueryPool->QueryStride,
 					TempResults.Num() * TempResults.GetTypeSize(),
 					(uint64_t*)TempResults.GetData(),
 					sizeof(FQueryResult),
@@ -687,13 +687,42 @@ void FVulkanDynamicRHI::CompletePayload(FVulkanPayload* Payload)
 				}
 			}
 
+			// In multiview, one "query" occupies multiple slots, one for each view.
+			// For queries issued with vkCmdBeginQuery/vkCmdEndQuery, the slots contain
+			// implementation-defined values, but their sum is guaranteed to represent
+			// the total of the query over all views. For queries issued with
+			// vkCmdWriteTimestamp, either the total over all views is written to the
+			// first slot and all other slots are 0, or each slot contains the timestamp
+			// for its associated view. In those cases, we want the maximum timestamp
+			// across all views.
+			const bool bUseMaxValue = QueryPool->GetPoolType() == EVulkanQueryPoolType::Timestamp;
+
 			for (uint32 QueryIndex = 0; QueryIndex < QueryPool->CurrentQueryCount; ++QueryIndex)
 			{
-				if (TempResults[QueryIndex].Availability)
-				{
 					uint64* ResultPtr = QueryPool->QueryResults[QueryIndex];
 					if (ResultPtr)
 					{
+					uint64 QueryResult = 0;
+					for (uint32 ViewIndex = 0; ViewIndex < QueryPool->QueryStride; ++ViewIndex)
+					{
+						uint32 ResultIndex = QueryIndex * QueryPool->QueryStride + ViewIndex;
+						if (TempResults[ResultIndex].Availability)
+						{
+							if (bUseMaxValue)
+							{
+								QueryResult = FMath::Max(QueryResult, TempResults[ResultIndex].Result);
+							}
+							else
+							{
+								QueryResult += TempResults[ResultIndex].Result;
+							}
+						}
+						else
+						{
+							UE_LOG(LogVulkanRHI, Warning, TEXT("Could not read back query: [PoolType=%d, Index=%d, View=%d]"), (int32)QueryPool->GetPoolType(), QueryIndex, ViewIndex);
+						}
+					}
+
 #if RHI_NEW_GPU_PROFILER
 						if (QueryPool->GetPoolType() == EVulkanQueryPoolType::Timestamp)
 						{
@@ -705,7 +734,7 @@ void FVulkanDynamicRHI::CompletePayload(FVulkanPayload* Payload)
 							uint64 GPUCalibrationTimestamp = CurrentTiming->GPUTimestamp;
 
 							// NOTE: Adreno GPU calibration values aren't valid, create a fallback
-							if (bBypassTiming || (TempResults[QueryIndex].Result < GPUCalibrationTimestamp) || (GPUCalibrationTimestamp == 0))
+						if (bBypassTiming || (QueryResult < GPUCalibrationTimestamp) || (GPUCalibrationTimestamp == 0))
 							{
 								if (!bBypassTiming)
 								{
@@ -717,14 +746,14 @@ void FVulkanDynamicRHI::CompletePayload(FVulkanPayload* Payload)
 											TEXT("GPU timestamp can't be moved to CPU timeline.  ")
 											TEXT("Calibration data is invalid (Curr=%lld  Cal=%lld)!  ")
 											TEXT("CPU and GPU profiling data will not line up!"),
-											TempResults[QueryIndex].Result, GPUCalibrationTimestamp);
+										QueryResult, GPUCalibrationTimestamp);
 									}
 
 									const uint64 CurrentTimestamp = FPlatformTime::Cycles64();
 									ReferenceCPUTimestamp = ((CPUCalibrationTimestamp != 0) && 
 										(CPUCalibrationTimestamp < CurrentTimestamp)) ?
 										CPUCalibrationTimestamp : CurrentTimestamp;
-									ReferenceGPUTimestamp = TempResults[QueryIndex].Result;
+								ReferenceGPUTimestamp = QueryResult;
 									bBypassTiming = true;
 								}
 
@@ -732,25 +761,20 @@ void FVulkanDynamicRHI::CompletePayload(FVulkanPayload* Payload)
 								GPUCalibrationTimestamp = ReferenceGPUTimestamp;
 
 								// :todo: Clamp results within the imperfect calibration timestamps (to prevent negative values below)
-								TempResults[QueryIndex].Result = FMath::Max(TempResults[QueryIndex].Result, GPUCalibrationTimestamp);
+							QueryResult = FMath::Max(QueryResult, GPUCalibrationTimestamp);
 							}
 
-							const uint64 GPUDelta = TempResults[QueryIndex].Result - GPUCalibrationTimestamp;
+						const uint64 GPUDelta = QueryResult - GPUCalibrationTimestamp;
 							const uint64 CPUDelta = (GPUDelta * CurrentTiming->CPUFrequency) / CurrentTiming->GPUFrequency;
 							(*ResultPtr) = CPUCalibrationTimestamp + CPUDelta;
 						}
 						else
 #endif
 						{
-							(*ResultPtr) = TempResults[QueryIndex].Result;
+						(*ResultPtr) = QueryResult;
 						}
 					}
 				}
-				else
-				{
-					UE_LOG(LogVulkanRHI, Warning, TEXT("Could not read back query: [PoolType=%d, Index=%d]"), (int32)QueryPool->GetPoolType(), QueryIndex);
-				}
-			}
 
 			QueryPool->QueryRefs.Empty(QueryPool->GetMaxQueries());
 			Device->ReleaseQueryPool(QueryPool);
