VulkanRHI: Fix potential null access

#rnx

diff --git a/Engine/Source/Runtime/VulkanRHI/Private/VulkanRenderpass.h b/Engine/Source/Runtime/VulkanRHI/Private/VulkanRenderpass.h
--- a/Engine/Source/Runtime/VulkanRHI/Private/VulkanRenderpass.h
+++ b/Engine/Source/Runtime/VulkanRHI/Private/VulkanRenderpass.h
@@ -80,10 +80,12 @@ inline void FVulkanAttachmentReference<VkAttachmentReference2>::SetDepthStencilA
 	{
 		bSupportsParallelRendering = false;
 	}
+	const bool bHasStencilReference = StencilReference && (StencilReference->stencilLayout != VK_IMAGE_LAYOUT_UNDEFINED);
+
 	sType = VK_STRUCTURE_TYPE_ATTACHMENT_REFERENCE_2;
-	pNext = (bSupportsParallelRendering && StencilReference && StencilReference->stencilLayout != VK_IMAGE_LAYOUT_UNDEFINED) ? StencilReference : nullptr;
+	pNext = (bSupportsParallelRendering && bHasStencilReference) ? StencilReference : nullptr;
 	attachment = AttachmentReferenceIn.attachment;
-	layout = bSupportsParallelRendering ? AttachmentReferenceIn.layout : VulkanRHI::GetMergedDepthStencilLayout(AttachmentReferenceIn.layout, StencilReference->stencilLayout);
+	layout = (bSupportsParallelRendering || !bHasStencilReference) ? AttachmentReferenceIn.layout : VulkanRHI::GetMergedDepthStencilLayout(AttachmentReferenceIn.layout, StencilReference->stencilLayout);
 	aspectMask = AspectMask;
 }
 
@@ -338,8 +340,8 @@ struct FVulkanAttachmentDescription<VkAttachmentDescription2>
 		storeOp = InDesc.storeOp;
 		stencilLoadOp = InDesc.stencilLoadOp;
 		stencilStoreOp = InDesc.stencilStoreOp;
-		initialLayout = bSupportsParallelRendering ? InDesc.initialLayout : VulkanRHI::GetMergedDepthStencilLayout(InDesc.initialLayout, InStencilDesc->stencilInitialLayout);
-		finalLayout = bSupportsParallelRendering ? InDesc.finalLayout : VulkanRHI::GetMergedDepthStencilLayout(InDesc.finalLayout, InStencilDesc->stencilFinalLayout);
+		initialLayout = (bSupportsParallelRendering || (InStencilDesc == nullptr)) ? InDesc.initialLayout : VulkanRHI::GetMergedDepthStencilLayout(InDesc.initialLayout, InStencilDesc->stencilInitialLayout);
+		finalLayout = (bSupportsParallelRendering || (InStencilDesc == nullptr)) ? InDesc.finalLayout : VulkanRHI::GetMergedDepthStencilLayout(InDesc.finalLayout, InStencilDesc->stencilFinalLayout);
 	}
 };
 
