VulkanRHI: Also accept "AllowSoftwareRendering" in Vulkan to match with D3D params.

#rnx

diff --git a/Engine/Source/Runtime/VulkanRHI/Private/VulkanRHI.cpp b/Engine/Source/Runtime/VulkanRHI/Private/VulkanRHI.cpp
--- a/Engine/Source/Runtime/VulkanRHI/Private/VulkanRHI.cpp
+++ b/Engine/Source/Runtime/VulkanRHI/Private/VulkanRHI.cpp
@@ -246,7 +246,7 @@ static VkPhysicalDevice SelectPhysicalDevice(VkInstance InInstance)
 	}
 
 	// Skip all CPU devices if they aren't permitted
-	const bool bAllowCPUDevices = FParse::Param(FCommandLine::Get(), TEXT("AllowCPUDevices"));
+	const bool bAllowCPUDevices = FParse::Param(FCommandLine::Get(), TEXT("AllowCPUDevices")) || FParse::Param(FCommandLine::Get(), TEXT("AllowSoftwareRendering"));
 	for (int32 Index = 0; Index < PhysicalDeviceInfos.Num(); ++Index)
 	{
 		if (!bAllowCPUDevices && (PhysicalDeviceInfos[Index].PhysicalDeviceProperties2.properties.deviceType == VK_PHYSICAL_DEVICE_TYPE_CPU))
