VulkanRHI: reduce the default performance cost of the in/out markers in the diag buffer when no specialized extensions are supported.  Add a different mode for when precision is required.

#rnx

#[fyi] [REDACTED_HANDLE], [REDACTED_HANDLE]

diff --git a/Engine/Source/Runtime/VulkanRHI/Private/VulkanQueue.cpp b/Engine/Source/Runtime/VulkanRHI/Private/VulkanQueue.cpp
--- a/Engine/Source/Runtime/VulkanRHI/Private/VulkanQueue.cpp
+++ b/Engine/Source/Runtime/VulkanRHI/Private/VulkanQueue.cpp
@@ -52,7 +52,8 @@ static TAutoConsoleVariable<int32> CVarVulkanDiagnosticBuffer(
 	TEXT("r.Vulkan.DiagnosticBuffer"),
 	1,
 	TEXT("0: Disable the diagnostic buffer\n")
-	TEXT("1: Enable the diagnostic buffer (default)"),
+	TEXT("1: Enable the diagnostic buffer with less perf impact (default)")
+	TEXT("2: Enable the diagnostic buffer with more precision"),
 	ECVF_ReadOnly
 );
 
@@ -728,6 +729,23 @@ FVulkanDiagnosticBuffer::~FVulkanDiagnosticBuffer()
 }
 
 #if WITH_RHI_BREADCRUMBS
+static void ManuallyWriteMarker(const VkCommandBuffer CommandBufferHandle, const VkBuffer BufferHandle, const uint64 WriteOffset, const uint32 BreadcrumbID)
+{
+	// When using precise barriers, make sure any work before or after can't overlap with the marker's write.
+	// When aiming for minimal cost, use a single barrier before the writing the marker (the fact we write markers on in and out should mean we're never more than off by one)
+	const bool bUsePreciseBarriers = (CVarVulkanDiagnosticBuffer->GetInt() > 1);
+	const VkPipelineStageFlags DestStage = bUsePreciseBarriers ? VK_PIPELINE_STAGE_TOP_OF_PIPE_BIT : VK_PIPELINE_STAGE_TRANSFER_BIT;
+	VulkanDynamicAPI::vkCmdPipelineBarrier(CommandBufferHandle, VK_PIPELINE_STAGE_BOTTOM_OF_PIPE_BIT, DestStage, 0, 0, nullptr, 0, nullptr, 0, nullptr);
+
+	VulkanDynamicAPI::vkCmdFillBuffer(CommandBufferHandle, BufferHandle, WriteOffset, sizeof(BreadcrumbID), BreadcrumbID);
+
+	if (bUsePreciseBarriers)
+	{
+		VulkanDynamicAPI::vkCmdPipelineBarrier(CommandBufferHandle, VK_PIPELINE_STAGE_BOTTOM_OF_PIPE_BIT, DestStage, 0, 0, nullptr, 0, nullptr, 0, nullptr);
+	}
+
+}
+
 void FVulkanDiagnosticBuffer::WriteMarkerIn(FVulkanCommandBuffer& CommandBuffer, FRHIBreadcrumbNode* Breadcrumb) const
 {
 	const VkCommandBuffer CommandBufferHandle = CommandBuffer.GetHandle();
@@ -742,9 +760,7 @@ void FVulkanDiagnosticBuffer::WriteMarkerIn(FVulkanCommandBuffer& CommandBuffer,
 		// Buffer writes cannot be recorded inside of a render pass
 		if (CommandBuffer.IsOutsideRenderPass())
 		{
-			VulkanDynamicAPI::vkCmdPipelineBarrier(CommandBufferHandle, VK_PIPELINE_STAGE_BOTTOM_OF_PIPE_BIT, VK_PIPELINE_STAGE_TOP_OF_PIPE_BIT, 0, 0, nullptr, 0, nullptr, 0, nullptr);
-			VulkanDynamicAPI::vkCmdFillBuffer(CommandBufferHandle, Buffer, WriteOffset, sizeof(Breadcrumb->ID), Breadcrumb->ID);
-			VulkanDynamicAPI::vkCmdPipelineBarrier(CommandBufferHandle, VK_PIPELINE_STAGE_BOTTOM_OF_PIPE_BIT, VK_PIPELINE_STAGE_TOP_OF_PIPE_BIT, 0, 0, nullptr, 0, nullptr, 0, nullptr);
+			ManuallyWriteMarker(CommandBufferHandle, Buffer, WriteOffset, Breadcrumb->ID);
 		}
 	}
 
@@ -770,9 +786,7 @@ void FVulkanDiagnosticBuffer::WriteMarkerOut(FVulkanCommandBuffer& CommandBuffer
 		// Buffer writes cannot be recorded inside of a render pass
 		if (CommandBuffer.IsOutsideRenderPass())
 		{
-			VulkanDynamicAPI::vkCmdPipelineBarrier(CommandBufferHandle, VK_PIPELINE_STAGE_BOTTOM_OF_PIPE_BIT, VK_PIPELINE_STAGE_TOP_OF_PIPE_BIT, 0, 0, nullptr, 0, nullptr, 0, nullptr);
-			VulkanDynamicAPI::vkCmdFillBuffer(CommandBufferHandle, Buffer, WriteOffset, sizeof(Breadcrumb->ID), Breadcrumb->ID);
-			VulkanDynamicAPI::vkCmdPipelineBarrier(CommandBufferHandle, VK_PIPELINE_STAGE_BOTTOM_OF_PIPE_BIT, VK_PIPELINE_STAGE_TOP_OF_PIPE_BIT, 0, 0, nullptr, 0, nullptr, 0, nullptr);
+			ManuallyWriteMarker(CommandBufferHandle, Buffer, WriteOffset, Breadcrumb->ID);
 		}
 	}
 
