VulkanRHI: Bring in part of the SyncPoint class from D3D (only used for CPU syncs for now, WIP).  Add a new CPU SyncPoint type that is ok to push back.  Reduce submit counts in very particular circumstances (single queue and single RHI thread).

#rnx
#[fyi] [REDACTED_HANDLE], [REDACTED_HANDLE], [REDACTED_HANDLE]

diff --git a/Engine/Config/Android/AndroidEngine.ini b/Engine/Config/Android/AndroidEngine.ini
--- a/Engine/Config/Android/AndroidEngine.ini
+++ b/Engine/Config/Android/AndroidEngine.ini
@@ -89,6 +89,7 @@ r.Vulkan.Submission.AllowTimelineSemaphores=0
 r.Vulkan.Submission.UseInterruptThread=0
 r.Vulkan.Submission.UseSubmissionThread=0
 r.Vulkan.RHIThread=1
+r.Vulkan.UseMinimalSubmits=1
 ; Disable async compute until we can target devices to enable it on
 r.Vulkan.AllowAsyncCompute=0
 ; Disable split barriers until its effects are measured on all devices
diff --git a/Engine/Source/Runtime/VulkanRHI/Private/VulkanContext.cpp b/Engine/Source/Runtime/VulkanRHI/Private/VulkanContext.cpp
--- a/Engine/Source/Runtime/VulkanRHI/Private/VulkanContext.cpp
+++ b/Engine/Source/Runtime/VulkanRHI/Private/VulkanContext.cpp
@@ -129,6 +129,11 @@ void FVulkanContextCommon::FlushCommands(EVulkanFlushFlags FlushFlags)
 		FVulkanDynamicRHI::Get().ProcessInterruptQueueUntil(SyncPoint);
 	}
 
+	if (Device.UseMinimalSubmits())
+	{
+		FVulkanDynamicRHI::Get().KickSubmissionThread(true);
+	}
+
 	if (SubmissionEvent && !SubmissionEvent->IsComplete())
 	{
 		SCOPED_NAMED_EVENT_TEXT("Submission_Wait", FColor::Turquoise);
diff --git a/Engine/Source/Runtime/VulkanRHI/Private/VulkanContext.h b/Engine/Source/Runtime/VulkanRHI/Private/VulkanContext.h
--- a/Engine/Source/Runtime/VulkanRHI/Private/VulkanContext.h
+++ b/Engine/Source/Runtime/VulkanRHI/Private/VulkanContext.h
@@ -184,11 +184,11 @@ public:
 	}
 
 	// Returns a single sync point for the context that will be inserted when it is finalized
-	FVulkanSyncPoint* GetContextSyncPoint()
+	FVulkanSyncPointRef GetContextSyncPoint()
 	{
 		if (!ContextSyncPoint)
 		{
-			ContextSyncPoint = CreateVulkanSyncPoint();
+			ContextSyncPoint = FVulkanSyncPoint::Create(EVulkanSyncPointType::Context, TEXT("VulkanContextSyncPoint"));
 		}
 		return ContextSyncPoint;
 	}
diff --git a/Engine/Source/Runtime/VulkanRHI/Private/VulkanDevice.cpp b/Engine/Source/Runtime/VulkanRHI/Private/VulkanDevice.cpp
--- a/Engine/Source/Runtime/VulkanRHI/Private/VulkanDevice.cpp
+++ b/Engine/Source/Runtime/VulkanRHI/Private/VulkanDevice.cpp
@@ -77,6 +77,13 @@ static TAutoConsoleVariable<int32> GCVarVulkanTempBlockSizeKB(
 	ECVF_ReadOnly
 );
 
+static TAutoConsoleVariable<int32> GCVarVulkanUseMinimalSubmits(
+	TEXT("r.Vulkan.UseMinimalSubmits"),
+	0,
+	TEXT("Can be specified in device profiles to keep submit counts low on devices with a single queue and single RHI thread."),
+	ECVF_RenderThreadSafe
+);
+
 // Mirror GPixelFormats with format information for buffers
 VkFormat GVulkanBufferFormat[PF_MAX];
 
@@ -493,6 +500,16 @@ void FVulkanDevice::CreateDevice(TArray<const ANSICHAR*>& DeviceLayers, FVulkanD
 	SetupDrawMarkers();
 }
 
+bool FVulkanDevice::UseMinimalSubmits() const
+{
+	return (ActiveQueueFamilies.Num() == 1) && 
+		!GRHIGlobals.SupportsParallelRHIExecute &&
+		!FVulkanDynamicRHI::Get().HasInterruptThread() &&
+		!FVulkanDynamicRHI::Get().HasSubmissionThread() &&
+		GCVarVulkanUseMinimalSubmits.GetValueOnRenderThread();
+}
+
+
 #if VULKAN_ENABLE_DRAW_MARKERS
 TAutoConsoleVariable<int32> CVarVulkanDebugMarkers(
 	TEXT("r.Vulkan.DebugMarkers"),
diff --git a/Engine/Source/Runtime/VulkanRHI/Private/VulkanDevice.h b/Engine/Source/Runtime/VulkanRHI/Private/VulkanDevice.h
--- a/Engine/Source/Runtime/VulkanRHI/Private/VulkanDevice.h
+++ b/Engine/Source/Runtime/VulkanRHI/Private/VulkanDevice.h
@@ -607,6 +607,8 @@ public:
 		return ActiveQueueFamilies;
 	}
 
+	bool UseMinimalSubmits() const;
+
 private:
 	void MapBufferFormatSupport(FPixelFormatInfo& PixelFormatInfo, EPixelFormat UEFormat, VkFormat VulkanFormat);
 	void MapImageFormatSupport(FPixelFormatInfo& PixelFormatInfo, const TArrayView<const VkFormat>& PrioritizedFormats, EPixelFormatCapabilities RequiredCapabilities);
diff --git a/Engine/Source/Runtime/VulkanRHI/Private/VulkanDynamicRHI.h b/Engine/Source/Runtime/VulkanRHI/Private/VulkanDynamicRHI.h
--- a/Engine/Source/Runtime/VulkanRHI/Private/VulkanDynamicRHI.h
+++ b/Engine/Source/Runtime/VulkanRHI/Private/VulkanDynamicRHI.h
@@ -328,6 +328,7 @@ public:
 
 	static void SavePipelineCache();
 
+	void ProcessInterruptQueueUntil(FVulkanSyncPointRef SyncPoint);
 	void ProcessInterruptQueueUntil(FGraphEvent* GraphEvent);
 	void EnqueueEndOfPipeTask(TUniqueFunction<void()> TaskFunc, TUniqueFunction<void(FVulkanPayload&)> ModifyPayloadCallback = {});
 	void CompletePayload(FVulkanPayload* Payload);
@@ -342,13 +343,24 @@ public:
 	void ProcessInterruptQueueOnGPUCrash();
 	void TerminateOnGPUCrash();
 
+	bool HasInterruptThread() const
+	{
+		return (InterruptThread != nullptr);
+	}
+
+	bool HasSubmissionThread() const
+	{
+		return (SubmissionThread != nullptr);
+	}
+
+	void KickSubmissionThread(bool bForce = false);
+
 protected:
 	void InitializeSubmissionPipe();
 	void ShutdownSubmissionPipe();
 	bool ProcessSubmissionQueue();
 	bool ProcessInterruptQueue();
 	bool WaitAndProcessInterruptQueue();
-	void KickSubmissionThread();
 	void KickInterruptThread();
 
 #if RHI_NEW_GPU_PROFILER
diff --git a/Engine/Source/Runtime/VulkanRHI/Private/VulkanGPUProfiler.h b/Engine/Source/Runtime/VulkanRHI/Private/VulkanGPUProfiler.h
--- a/Engine/Source/Runtime/VulkanRHI/Private/VulkanGPUProfiler.h
+++ b/Engine/Source/Runtime/VulkanRHI/Private/VulkanGPUProfiler.h
@@ -16,7 +16,7 @@ class FVulkanDevice;
 
 #if (RHI_NEW_GPU_PROFILER == 0)
 
-using FVulkanSyncPoint = FGraphEvent;
+class FVulkanSyncPoint;
 using FVulkanSyncPointRef = TRefCountPtr<FVulkanSyncPoint>;
 
 class FVulkanGPUTiming : public FGPUTiming
diff --git a/Engine/Source/Runtime/VulkanRHI/Private/VulkanMemory.cpp b/Engine/Source/Runtime/VulkanRHI/Private/VulkanMemory.cpp
--- a/Engine/Source/Runtime/VulkanRHI/Private/VulkanMemory.cpp
+++ b/Engine/Source/Runtime/VulkanRHI/Private/VulkanMemory.cpp
@@ -1981,7 +1981,7 @@ namespace VulkanRHI
 					}
 				}
 			}
-			AllocationSize = FMath::Max(Size, MemoryBucket.PageSize); // for allocations above max, which are forced to be seperate allocations
+			AllocationSize = FMath::Max(Size, MemoryBucket.PageSize); // for allocations above max, which are forced to be separate allocations
 		}
 		else
 		{
@@ -4301,6 +4301,10 @@ namespace VulkanRHI
 		Device->GetMemoryManager().FreeVulkanAllocation(Allocation);
 	}
 
+	FStagingManager::FStagingManager()
+	{
+	}
+
 	FStagingManager::~FStagingManager()
 	{
 	}
diff --git a/Engine/Source/Runtime/VulkanRHI/Private/VulkanMemory.h b/Engine/Source/Runtime/VulkanRHI/Private/VulkanMemory.h
--- a/Engine/Source/Runtime/VulkanRHI/Private/VulkanMemory.h
+++ b/Engine/Source/Runtime/VulkanRHI/Private/VulkanMemory.h
@@ -36,7 +36,7 @@ class FVulkanTexture;
 
 struct FVulkanContextArray;
 
-using FVulkanSyncPoint = FGraphEvent;
+class FVulkanSyncPoint;
 using FVulkanSyncPointRef = TRefCountPtr<FVulkanSyncPoint>;
 
 #if !VULKAN_OBJECT_TRACKING
@@ -908,12 +908,7 @@ namespace VulkanRHI
 	class FStagingManager
 	{
 	public:
-		FStagingManager() :
-			PeakUsedMemory(0),
-			UsedMemory(0),
-			Device(nullptr)
-		{
-		}
+		FStagingManager();
 		~FStagingManager();
 
 		void Init(FVulkanDevice* InDevice)
@@ -930,7 +925,6 @@ namespace VulkanRHI
 
 		void ProcessPendingFree(bool bImmediately, bool bFreeToOS);
 
-
 		void GetMemoryDump(FResourceHeapStats& Stats);
 		void DumpMemory();
 
@@ -948,12 +942,12 @@ namespace VulkanRHI
 		};
 		TArray<FFreeEntry> FreeStagingBuffers;
 
-		uint64 PeakUsedMemory;
-		uint64 UsedMemory;
+		uint64 PeakUsedMemory = 0;
+		uint64 UsedMemory = 0;
 
 		void ProcessPendingFreeNoLock(bool bImmediately, bool bFreeToOS);
 
-		FVulkanDevice* Device;
+		FVulkanDevice* Device = nullptr;
 	};
 
 	// Simple temp allocation blocks used for volatile allocations
diff --git a/Engine/Source/Runtime/VulkanRHI/Private/VulkanQuery.cpp b/Engine/Source/Runtime/VulkanRHI/Private/VulkanQuery.cpp
--- a/Engine/Source/Runtime/VulkanRHI/Private/VulkanQuery.cpp
+++ b/Engine/Source/Runtime/VulkanRHI/Private/VulkanQuery.cpp
@@ -269,7 +269,7 @@ void FVulkanDynamicRHI::RHIEndRenderQuery_TopOfPipe(FRHICommandListBase& RHICmdL
 	{
 		// Queries issued outside of a batch use one sync point per query.
 //		check(Query->SyncPoint == nullptr);
-		Query->SyncPoint = CreateVulkanSyncPoint();
+		Query->SyncPoint = FVulkanSyncPoint::Create(EVulkanSyncPointType::GPUAndCPU, TEXT("UnbatchedRenderQuerySyncPoint"));
 
 		RHICmdList.EnqueueLambda([SyncPoint = Query->SyncPoint](FRHICommandListBase& ExecutingCmdList)
 		{
@@ -290,7 +290,7 @@ void FVulkanDynamicRHI::RHIBeginRenderQueryBatch_TopOfPipe(FRHICommandListBase&
 	auto& QueryBatchData = RHICmdList.GetQueryBatchData(QueryType);
 	checkf(QueryBatchData[GPUIndex] == nullptr, TEXT("A query batch for this type has already begun on this command list."));
 
-	FVulkanSyncPointRef SyncPoint = CreateVulkanSyncPoint();
+	FVulkanSyncPointRef SyncPoint = FVulkanSyncPoint::Create(EVulkanSyncPointType::GPUAndCPU, TEXT("BatchedRenderQuerySyncPoint"));
 
 	// Keep a reference in the RHI command list, so we can retrieve it later in BeginQuery/EndQuery/EndBatch.
 	QueryBatchData[GPUIndex] = SyncPoint.GetReference();
diff --git a/Engine/Source/Runtime/VulkanRHI/Private/VulkanQuery.h b/Engine/Source/Runtime/VulkanRHI/Private/VulkanQuery.h
--- a/Engine/Source/Runtime/VulkanRHI/Private/VulkanQuery.h
+++ b/Engine/Source/Runtime/VulkanRHI/Private/VulkanQuery.h
@@ -15,7 +15,7 @@ class FVulkanCommandBuffer;
 class FVulkanRenderQuery;
 class FVulkanDynamicRHI;
 
-using FVulkanSyncPoint = FGraphEvent;
+class FVulkanSyncPoint;
 using FVulkanSyncPointRef = TRefCountPtr<FVulkanSyncPoint>;
 
 
diff --git a/Engine/Source/Runtime/VulkanRHI/Private/VulkanQueue.cpp b/Engine/Source/Runtime/VulkanRHI/Private/VulkanQueue.cpp
--- a/Engine/Source/Runtime/VulkanRHI/Private/VulkanQueue.cpp
+++ b/Engine/Source/Runtime/VulkanRHI/Private/VulkanQueue.cpp
@@ -282,6 +282,40 @@ void FVulkanQueue::SubmitPayloads(TArrayView<FVulkanPayload*> Payloads, TMap<VkS
 	CommandBufferStorage.Reserve(NumCommandBuffers);
 	SubmitInfos.Reserve(NumPayloads);
 
+	if (Device.UseMinimalSubmits())
+	{
+		for (int32 PayloadIndex = Payloads.Num()-1; PayloadIndex >= 1; PayloadIndex--)
+		{
+			FVulkanPayload* PreviousPayload = Payloads[PayloadIndex - 1];
+			FVulkanPayload* CurrentPayload = Payloads[PayloadIndex];
+
+			// We can't merge if we have syncs or action that need to happen between the command buffers
+			if (CurrentPayload->WaitSemaphores.IsEmpty() && 
+				!CurrentPayload->PreExecuteCallback &&
+				PreviousPayload->SignalSemaphores.IsEmpty() &&
+				PreviousPayload->ReservedResourcesToCommit.IsEmpty() &&
+				!PreviousPayload->Timing.IsSet())
+			{
+				bool bMovableSyncs = true;
+				for (FVulkanSyncPointRef& SyncPoint : PreviousPayload->SyncPoints)
+				{
+					if (SyncPoint->GetType() != EVulkanSyncPointType::Context)
+					{
+						bMovableSyncs = false;
+						break;
+					}
+				}
+
+				if (!bMovableSyncs)
+				{
+					continue;
+				}
+
+				PreviousPayload->CommandBuffers.Append(MoveTemp(CurrentPayload->CommandBuffers));
+			}
+		}
+	}
+
 	const uint64 SubmitTime = FPlatformTime::Cycles64();
 
 	for (FVulkanPayload* Payload : Payloads)
@@ -451,6 +485,7 @@ void FVulkanQueue::SubmitPayloads(TArrayView<FVulkanPayload*> Payloads, TMap<VkS
 void FVulkanQueue::Submit(TArrayView<VkSubmitInfo> InSubmitInfos, FVulkanFence* Fence)
 {
 	SCOPE_CYCLE_COUNTER(STAT_VulkanQueueSubmit);
+	SCOPED_NAMED_EVENT(VulkanQueueSubmit, FColor::Purple);
 
 	const VkFence FenceHandle = Fence ? Fence->GetHandle() : VK_NULL_HANDLE;
 	VERIFYVULKANRESULT(VulkanRHI::vkQueueSubmit(Queue, InSubmitInfos.Num(), InSubmitInfos.GetData(), FenceHandle));
diff --git a/Engine/Source/Runtime/VulkanRHI/Private/VulkanRayTracing.cpp b/Engine/Source/Runtime/VulkanRHI/Private/VulkanRayTracing.cpp
--- a/Engine/Source/Runtime/VulkanRHI/Private/VulkanRayTracing.cpp
+++ b/Engine/Source/Runtime/VulkanRHI/Private/VulkanRayTracing.cpp
@@ -1780,6 +1780,12 @@ FVulkanRayTracingCompactionRequestHandler::FVulkanRayTracingCompactionRequestHan
 	ActiveBLASes.Reserve(GVulkanRayTracingMaxBatchedCompaction);
 }
 
+FVulkanRayTracingCompactionRequestHandler::~FVulkanRayTracingCompactionRequestHandler()
+{
+	check(PendingRequests.IsEmpty());
+	delete QueryPool;
+}
+
 void FVulkanRayTracingCompactionRequestHandler::RequestCompact(FVulkanRayTracingGeometry* InRTGeometry)
 {
 	check(InRTGeometry->AccelerationStructureBuffer);
diff --git a/Engine/Source/Runtime/VulkanRHI/Private/VulkanRayTracing.h b/Engine/Source/Runtime/VulkanRHI/Private/VulkanRayTracing.h
--- a/Engine/Source/Runtime/VulkanRHI/Private/VulkanRayTracing.h
+++ b/Engine/Source/Runtime/VulkanRHI/Private/VulkanRayTracing.h
@@ -12,7 +12,7 @@ class FVulkanCommandListContext;
 class FVulkanBuffer;
 struct FVulkanRayTracingGeometryParameters;
 
-using FVulkanSyncPoint = FGraphEvent;
+class FVulkanSyncPoint;
 using FVulkanSyncPointRef = TRefCountPtr<FVulkanSyncPoint>;
 
 class FVulkanRayTracingPlatform
@@ -336,11 +336,7 @@ public:
 	UE_NONCOPYABLE(FVulkanRayTracingCompactionRequestHandler)
 
 	FVulkanRayTracingCompactionRequestHandler(FVulkanDevice& InDevice);
-	~FVulkanRayTracingCompactionRequestHandler()
-	{
-		check(PendingRequests.IsEmpty());
-		delete QueryPool;
-	}
+	~FVulkanRayTracingCompactionRequestHandler();
 
 	void RequestCompact(FVulkanRayTracingGeometry* InRTGeometry);
 	bool ReleaseRequest(FVulkanRayTracingGeometry* InRTGeometry);
diff --git a/Engine/Source/Runtime/VulkanRHI/Private/VulkanResources.h b/Engine/Source/Runtime/VulkanRHI/Private/VulkanResources.h
--- a/Engine/Source/Runtime/VulkanRHI/Private/VulkanResources.h
+++ b/Engine/Source/Runtime/VulkanRHI/Private/VulkanResources.h
@@ -35,7 +35,7 @@ class FVulkanRenderQuery;
 class FVulkanDynamicRHI;
 struct FVulkanContextArray;
 
-using FVulkanSyncPoint = FGraphEvent;
+class FVulkanSyncPoint;
 using FVulkanSyncPointRef = TRefCountPtr<FVulkanSyncPoint>;
 
 namespace VulkanRHI
@@ -672,7 +672,7 @@ public:
 		VkImageCompressionControlEXT CompressionControl;
 	};
 
-	// Seperate method for creating VkImageCreateInfo
+	// Separate method for creating VkImageCreateInfo
 	static void GenerateImageCreateInfo(
 		FImageCreateInfo& OutImageCreateInfo,
 		FVulkanDevice& InDevice,
diff --git a/Engine/Source/Runtime/VulkanRHI/Private/VulkanSubmission.cpp b/Engine/Source/Runtime/VulkanRHI/Private/VulkanSubmission.cpp
--- a/Engine/Source/Runtime/VulkanRHI/Private/VulkanSubmission.cpp
+++ b/Engine/Source/Runtime/VulkanRHI/Private/VulkanSubmission.cpp
@@ -17,6 +17,25 @@
 
 extern CORE_API bool GIsGPUCrashed;
 
+TLockFreePointerListUnordered<void, PLATFORM_CACHE_LINE_SIZE> FVulkanSyncPoint::MemoryPool;
+
+
+void FVulkanSyncPoint::Wait() const
+{
+	checkf(GraphEvent, TEXT("This sync point was not created with a CPU event. Cannot wait for completion on the CPU."));
+
+	if (!GraphEvent->IsComplete())
+	{
+		// Block the calling thread until the graph event is signaled by the interrupt thread.
+		SCOPED_NAMED_EVENT_TEXT("SyncPoint_Wait", FColor::Turquoise);
+		FVulkanDynamicRHI::Get().ProcessInterruptQueueUntil(GraphEvent);
+	}
+
+	check(GraphEvent->IsComplete());
+}
+
+
+
 FVulkanPayload::~FVulkanPayload()
 {
 	for (FVulkanSemaphore* Semaphore : WaitSemaphores)
@@ -50,7 +69,6 @@ void FVulkanPayload::Merge(FVulkanPayload& OtherPayload)
 }
 
 
-
 class FVulkanThread final : private FRunnable
 {
 	typedef bool(FVulkanDynamicRHI::* FQueueFunc)();
@@ -213,13 +231,13 @@ void FVulkanDynamicRHI::KickInterruptThread()
 	}
 }
 
-void FVulkanDynamicRHI::KickSubmissionThread()
+void FVulkanDynamicRHI::KickSubmissionThread(bool bForce)
 {
 	if (SubmissionThread)
 	{
 		SubmissionThread->Kick();
 	}
-	else
+	else if (bForce || !Device->UseMinimalSubmits())
 	{
 		FScopeLock Lock(&SubmissionCS);
 		while (ProcessSubmissionQueue()) {};
@@ -482,6 +500,12 @@ bool FVulkanDynamicRHI::ProcessInterruptQueue()
 	return bProgress;
 }
 
+void FVulkanDynamicRHI::ProcessInterruptQueueUntil(FVulkanSyncPointRef SyncPoint)
+{
+	checkf(SyncPoint.IsValid(), TEXT("Waiting on invalid sync point."));
+	ProcessInterruptQueueUntil(SyncPoint->GetGraphEvent());
+}
+
 void FVulkanDynamicRHI::ProcessInterruptQueueUntil(FGraphEvent* GraphEvent)
 {
 	if (InterruptThread)
@@ -493,6 +517,11 @@ void FVulkanDynamicRHI::ProcessInterruptQueueUntil(FGraphEvent* GraphEvent)
 	}
 	else if (GraphEvent)
 	{
+		if (Device->UseMinimalSubmits())
+		{
+			KickSubmissionThread(true);
+		}
+
 		// If we're waiting for an event, accumulate the idle time
 		UE::Stats::FThreadIdleStats::FScopeIdle IdleScope;
 
@@ -559,9 +588,9 @@ void FVulkanDynamicRHI::EnqueueEndOfPipeTask(TUniqueFunction<void()> TaskFunc, T
 	{
 		FVulkanPayload* Payload = new FVulkanPayload(Queue);
 
-		FVulkanSyncPointRef SyncPoint = CreateVulkanSyncPoint();
+		FVulkanSyncPointRef SyncPoint = FVulkanSyncPoint::Create(EVulkanSyncPointType::GPUAndCPU, TEXT("VulkanEndOfPipeTask"));
 		Payload->SyncPoints.Add(SyncPoint);
-		Prereqs.Add(SyncPoint);
+		Prereqs.Add(SyncPoint->GetGraphEvent());
 
 		if (ModifyPayloadCallback)
 		{
@@ -595,7 +624,7 @@ void FVulkanDynamicRHI::RHIBlockUntilGPUIdle()
 		Device->ForEachQueue([&EndOfQueueEvent, &Payloads](FVulkanQueue& Queue)
 		{
 			FVulkanPayload* Payload = new FVulkanPayload(Queue);
-			FVulkanSyncPointRef SyncPoint = CreateVulkanSyncPoint();
+			FVulkanSyncPointRef SyncPoint = FVulkanSyncPoint::Create(EVulkanSyncPointType::GPUAndCPU, TEXT("BlockUntilGPUIdle"));
 			Payload->SyncPoints.Add(SyncPoint);
 			Payloads->Add(Payload);
 
@@ -852,7 +881,7 @@ void FVulkanDynamicRHI::CompletePayload(FVulkanPayload* Payload)
 
 	for (FVulkanSyncPointRef& SyncPoint : Payload->SyncPoints)
 	{
-		SyncPoint->DispatchSubsequents();
+		SyncPoint->GetGraphEvent()->DispatchSubsequents();
 	}
 	Payload->SyncPoints.Empty();
 }
@@ -896,7 +925,7 @@ void FVulkanDynamicRHI::RHIRunOnQueue(EVulkanRHIRunOnQueueType QueueType, TFunct
 	}
 
 	PendingPayloadsForSubmission.Enqueue(Payloads);
-	KickSubmissionThread();
+	KickSubmissionThread(bWaitForSubmission);
 
 	// Use this opportunity to pump the interrupt queue
 	ProcessInterruptQueueUntil(nullptr);
diff --git a/Engine/Source/Runtime/VulkanRHI/Private/VulkanSubmission.h b/Engine/Source/Runtime/VulkanRHI/Private/VulkanSubmission.h
--- a/Engine/Source/Runtime/VulkanRHI/Private/VulkanSubmission.h
+++ b/Engine/Source/Runtime/VulkanRHI/Private/VulkanSubmission.h
@@ -20,14 +20,119 @@ class FVulkanSemaphore;
 class FVulkanTiming;
 class FVulkanDescriptorPoolSetContainer;
 
-// NOTE: typedef to FGraphEvent for now, but eventually move all syncs inside a FVulkanSyncPoint class
-using FVulkanSyncPoint = FGraphEvent;
+
+enum class EVulkanSyncPointType
+{
+	// Sync points of this type do not include an FGraphEvent, so cannot
+	// report completion to the CPU (via either IsComplete() or Wait())
+	// The will resolve into binary or timeline semaphores at submit time.
+	//GPUOnly,  :todo-jn: Move GPU syncs to new FVulkanSyncPoint class
+
+	// Sync points of this type include an FGraphEvent. The IsComplete() and Wait() functions
+	// can be used to poll for completion from the CPU, or block the CPU, respectively.
+	GPUAndCPU,
+
+	// Same as GPUAndCPU but not time critical and can be pushed back up until the end of the frame.
+	// Generally meant to be used for work that can happen at any point after the command buffers
+	// are done executing on the GPU.
+	Context,
+};
+
+class FVulkanSyncPoint;
 using FVulkanSyncPointRef = TRefCountPtr<FVulkanSyncPoint>;
 
-static FVulkanSyncPointRef CreateVulkanSyncPoint()
+//
+// A sync point is a logical point on a GPU queue's timeline that can be awaited by other queues, or the CPU.
+// These are used throughout the RHI as a way to abstract the underlying Vulkan semaphores and fences. 
+// The submission thread manages the underlying semaphores, fences and signaled values, and reports completion 
+// to the relevant sync points via an FGraphEvent.
+//
+// Sync points are one-shot, meaning they represent a single timeline point, and are released after use, via ref-counting.
+// Use FVulkanSyncPoint::Create() to make a new sync point and hold a reference to it via a FVulkanSyncPointRef object.
+//
+class FVulkanSyncPoint final : public FThreadSafeRefCountedObject
 {
-	return FGraphEvent::CreateGraphEvent();
-}
+	friend FVulkanDynamicRHI;
+	friend FVulkanQueue;
+
+	static TLockFreePointerListUnordered<void, PLATFORM_CACHE_LINE_SIZE> MemoryPool;
+
+	// No copying or moving
+	FVulkanSyncPoint(FVulkanSyncPoint const&) = delete;
+	FVulkanSyncPoint(FVulkanSyncPoint&&) = delete;
+
+	FVulkanSyncPoint(EVulkanSyncPointType InType, const TCHAR* InDebugName)
+		: Type(InType)
+#if RHI_USE_SYNC_POINT_DEBUG_NAME
+		, DebugName(InDebugName)
+#endif
+	{
+		if ((Type == EVulkanSyncPointType::GPUAndCPU) || (Type == EVulkanSyncPointType::Context))
+		{
+			GraphEvent = FGraphEvent::CreateGraphEvent();
+		}
+	}
+
+public:
+	static FVulkanSyncPointRef Create(EVulkanSyncPointType Type, const TCHAR* InDebugName)
+	{
+		LLM_SCOPE_BYNAME(TEXT("RHIMisc/CreateSyncPoint"));
+		return new FVulkanSyncPoint(Type, InDebugName);
+	}
+
+	bool IsComplete() const
+	{
+		checkf(GraphEvent, TEXT("This sync point was not created with a CPU event. Cannot check completion on the CPU."));
+		return GraphEvent->IsComplete();
+	}
+
+	void Wait() const;
+
+	FGraphEvent* GetGraphEvent() const
+	{
+		checkf(GraphEvent, TEXT("This sync point was not created with a CPU event."));
+		return GraphEvent;
+	}
+
+	EVulkanSyncPointType GetType() const
+	{
+		return Type;
+	}
+
+	void* operator new(size_t Size)
+	{
+		check(Size == sizeof(FVulkanSyncPoint));
+
+		void* Memory = MemoryPool.Pop();
+		if (!Memory)
+		{
+			Memory = FMemory::Malloc(sizeof(FVulkanSyncPoint), alignof(FVulkanSyncPoint));
+		}
+		return Memory;
+	}
+
+	void operator delete(void* Pointer)
+	{
+		MemoryPool.Push(Pointer);
+	}
+
+	inline const TCHAR* GetDebugName() const
+	{
+#if RHI_USE_SYNC_POINT_DEBUG_NAME
+		return DebugName;
+#else
+		return TEXT("SyncPointBusyWait");
+#endif
+	}
+
+private:
+	FGraphEventRef GraphEvent;
+	EVulkanSyncPointType Type;
+
+#if RHI_USE_SYNC_POINT_DEBUG_NAME
+	const TCHAR* DebugName = nullptr;
+#endif
+};
 
 
 
diff --git a/Engine/Source/Runtime/VulkanRHI/Private/VulkanSwapChain.cpp b/Engine/Source/Runtime/VulkanRHI/Private/VulkanSwapChain.cpp
--- a/Engine/Source/Runtime/VulkanRHI/Private/VulkanSwapChain.cpp
+++ b/Engine/Source/Runtime/VulkanRHI/Private/VulkanSwapChain.cpp
@@ -793,6 +793,11 @@ FVulkanSwapChain::EStatus FVulkanSwapChain::Present(FVulkanQueue* PresentQueue,
 {
 	check(CurrentImageIndex != -1);
 
+	if (Device.UseMinimalSubmits())
+	{
+		GVulkanRHI->KickSubmissionThread(true);
+	}
+
 	VkPresentInfoKHR Info;
 	ZeroVulkanStruct(Info, VK_STRUCTURE_TYPE_PRESENT_INFO_KHR);
 	VkSemaphore Semaphore = VK_NULL_HANDLE;
@@ -845,6 +850,7 @@ FVulkanSwapChain::EStatus FVulkanSwapChain::Present(FVulkanQueue* PresentQueue,
 
 	{
 		SCOPE_CYCLE_COUNTER(STAT_VulkanQueuePresent);
+		SCOPED_NAMED_EVENT(VulkanQueuePresent, FColor::Yellow);
 
 		VkResult PresentResult;
 		{
diff --git a/Engine/Source/Runtime/VulkanRHI/Private/VulkanUtil.cpp b/Engine/Source/Runtime/VulkanRHI/Private/VulkanUtil.cpp
--- a/Engine/Source/Runtime/VulkanRHI/Private/VulkanUtil.cpp
+++ b/Engine/Source/Runtime/VulkanRHI/Private/VulkanUtil.cpp
@@ -690,7 +690,10 @@ void FVulkanDynamicRHI::RHIWriteGPUFence_TopOfPipe(FRHICommandListBase& RHICmdLi
 	check(Fence);
 
 	checkf(!Fence->SubmittedSyncPoint.IsValid() && !Fence->CompletedSyncPoint.IsValid(), TEXT("The fence for the current GPU node has already been issued."));
-	Fence->CompletedSyncPoint = CreateVulkanSyncPoint();
+
+	// When using minimal submits, the simple threading scheme means we can push back when fence results are communicated
+	const EVulkanSyncPointType SyncPointType = Device->UseMinimalSubmits() ? EVulkanSyncPointType::Context : EVulkanSyncPointType::GPUAndCPU;
+	Fence->CompletedSyncPoint = FVulkanSyncPoint::Create(SyncPointType, TEXT("GPUFence"));
 
 	if (Fence->Event)
 	{
