VulkanRHI: Fix swapchain recreation when surface is lost during swapchain creation instead of Present.  Move VkSwapChain creation from the constructor to a function that can return a failure.  Remove unused functions and parameters.

#rnx

diff --git a/Engine/Source/Runtime/VulkanRHI/Private/VulkanContext.h b/Engine/Source/Runtime/VulkanRHI/Private/VulkanContext.h
--- a/Engine/Source/Runtime/VulkanRHI/Private/VulkanContext.h
+++ b/Engine/Source/Runtime/VulkanRHI/Private/VulkanContext.h
@@ -416,7 +416,6 @@ public:
 	bool AcquirePoolSetAndDescriptorsIfNeeded(const class FVulkanDescriptorSetsLayout& Layout, bool bNeedDescriptors, VkDescriptorSet* OutDescriptors);
 
 	bool IsSwapchainImage(FRHITexture* InTexture) const;
-	VkFormat GetSwapchainImageFormat() const;
 	FVulkanSwapChain* GetSwapChain() const;
 
 	FVulkanRenderPass* PrepareRenderPassForPSOCreation(const FGraphicsPipelineStateInitializer& Initializer);
diff --git a/Engine/Source/Runtime/VulkanRHI/Private/VulkanRHI.cpp b/Engine/Source/Runtime/VulkanRHI/Private/VulkanRHI.cpp
--- a/Engine/Source/Runtime/VulkanRHI/Private/VulkanRHI.cpp
+++ b/Engine/Source/Runtime/VulkanRHI/Private/VulkanRHI.cpp
@@ -1272,15 +1272,10 @@ void FVulkanCommandListContext::RHIEndDrawingViewport(FRHIViewport* ViewportRHI,
 	check(IsImmediate());
 	FVulkanViewport* Viewport = ResourceCast(ViewportRHI);
 
-	//#todo-rco: Unbind all pending state
-/*
-	check(bPresent);
-	RHI->Present();
-*/
 	FVulkanCommandBuffer& CommandBuffer = GetCommandBuffer();
 	check(!CommandBuffer.HasEnded() && !CommandBuffer.IsInsideRenderPass());
 
-	bool bNativePresent = Viewport->Present(*this, Device.GetPresentQueue(), bLockToVsync);
+	const bool bNativePresent = Viewport->Present(*this, bLockToVsync);
 	if (bNativePresent)
 	{
 		//#todo-rco: Check for r.FinishCurrentFrame
diff --git a/Engine/Source/Runtime/VulkanRHI/Private/VulkanRenderTarget.cpp b/Engine/Source/Runtime/VulkanRHI/Private/VulkanRenderTarget.cpp
--- a/Engine/Source/Runtime/VulkanRHI/Private/VulkanRenderTarget.cpp
+++ b/Engine/Source/Runtime/VulkanRHI/Private/VulkanRenderTarget.cpp
@@ -464,17 +464,6 @@ void FVulkanDynamicRHI::RHIRead3DSurfaceFloatData(FRHITexture* TextureRHI, FIntR
 	Device->GetStagingManager().ReleaseBuffer(nullptr, StagingBuffer);
 }
 
-VkFormat FVulkanCommandListContext::GetSwapchainImageFormat() const
-{
-	TArray<FVulkanViewport*>& viewports = FVulkanDynamicRHI::Get().GetViewports();
-	if (viewports.Num() == 0)
-	{
-		return VK_FORMAT_UNDEFINED;
-	}
-
-	return viewports[0]->GetSwapchainImageFormat();
-}
-
 FVulkanSwapChain* FVulkanCommandListContext::GetSwapChain() const
 {
 	TArray<FVulkanViewport*>& viewports = FVulkanDynamicRHI::Get().GetViewports();
diff --git a/Engine/Source/Runtime/VulkanRHI/Private/VulkanSwapChain.cpp b/Engine/Source/Runtime/VulkanRHI/Private/VulkanSwapChain.cpp
--- a/Engine/Source/Runtime/VulkanRHI/Private/VulkanSwapChain.cpp
+++ b/Engine/Source/Runtime/VulkanRHI/Private/VulkanSwapChain.cpp
@@ -90,7 +90,8 @@ static FAutoConsoleVariableRef CVarVulkanDebugVsync(
 
 #if !UE_BUILD_SHIPPING
 
-bool GSimulateLostSurfaceInNextTick = false;
+bool GSimulateLostSurfaceInNextPresent = false;
+bool GSimulateLostSurfaceInNextSwapChain = false;
 bool GSimulateSuboptimalSurfaceInNextTick = false;
 
 // A self registering exec helper to check for the VULKAN_* commands.
@@ -100,8 +101,8 @@ class FVulkanCommandsHelper : public FSelfRegisteringExec
 	{
 		if (FParse::Command(&Cmd, TEXT("VULKAN_SIMULATE_LOST_SURFACE")))
 		{
-			GSimulateLostSurfaceInNextTick = true;
-			Ar.Log(FString::Printf(TEXT("Vulkan: simulating lost surface next frame")));
+			GSimulateLostSurfaceInNextPresent = true;
+			Ar.Log(FString::Printf(TEXT("Vulkan: simulating lost surface next frame's present")));
 			return true;
 		}
 		else if (FParse::Command(&Cmd, TEXT("VULKAN_SIMULATE_SUBOPTIMAL_SURFACE")))
@@ -110,6 +111,12 @@ class FVulkanCommandsHelper : public FSelfRegisteringExec
 			Ar.Log(FString::Printf(TEXT("Vulkan: simulating suboptimal surface next frame")));
 			return true;
 		}
+		else if (FParse::Command(&Cmd, TEXT("VULKAN_SIMULATE_SWAPCHAIN_LOST_SURFACE")))
+		{
+			GSimulateLostSurfaceInNextSwapChain = true;
+			Ar.Log(FString::Printf(TEXT("Vulkan: simulating lost surface on next swapchain creation")));
+			return true;
+		}
 		else
 		{
 			return false;
@@ -120,9 +127,9 @@ static FVulkanCommandsHelper GVulkanCommandsHelper;
 
 VkResult SimulateErrors(VkResult Result)
 {
-	if (GSimulateLostSurfaceInNextTick)
+	if (GSimulateLostSurfaceInNextPresent)
 	{
-		GSimulateLostSurfaceInNextTick = false;
+		GSimulateLostSurfaceInNextPresent = false;
 		return VK_ERROR_SURFACE_LOST_KHR;
 	}
 
@@ -148,20 +155,20 @@ bool IsVulkanPreTransformEnabled(const FStaticShaderPlatform Platform)
 			&& CVarVulkanPreTransform.GetValueOnAnyThread() > 0;
 }
 
-FVulkanSwapChain::FVulkanSwapChain(VkInstance InInstance, FVulkanDevice& InDevice, EPixelFormat& InOutPixelFormat, uint32 Width, uint32 Height, bool bIsFullScreen,
-	uint32* InOutDesiredNumBackBuffers, TArray<VkImage>& OutImages, int8 InLockToVsync, FVulkanPlatformWindowContext& WindowContext, FVulkanSwapChainRecreateInfo* RecreateInfo)
-	: SwapChain(VK_NULL_HANDLE)
-	, Device(InDevice)
-	, Surface(VK_NULL_HANDLE)
-	, WindowHandle(WindowContext.GetWindowHandle())
-	, CurrentImageIndex(-1)
-	, SemaphoreIndex(0)
-	, NumPresentCalls(0)
-	, NumAcquireCalls(0)
-	, Instance(InInstance)
-	, LockToVsync(InLockToVsync)
+FVulkanSwapChain* FVulkanSwapChain::Create(
+	VkInstance InInstance, FVulkanDevice& Device, EPixelFormat& InOutPixelFormat, 
+	uint32 Width, uint32 Height, bool bIsFullScreen,
+	uint32* InOutDesiredNumBackBuffers, TArray<VkImage>& OutImages, int8 LockToVsync, 
+	FVulkanPlatformWindowContext& WindowContext, FVulkanSwapChainRecreateInfo* RecreateInfo)
 {
-	NextPresentTargetTime = (FPlatformTime::Seconds() - GStartTime);
+#define VERIFYVULKANSURFACERESULT(VkFunction) \
+	{\
+		const VkResult ScopedResult = VkFunction; \
+		if (ScopedResult == VK_ERROR_SURFACE_LOST_KHR) { return nullptr; } \
+		if (ScopedResult < VK_SUCCESS) { VulkanRHI::VerifyVulkanResult(ScopedResult, #VkFunction, __FILE__, __LINE__); } \
+	}
+
+	VkSurfaceKHR Surface = VK_NULL_HANDLE;
 
 	if (RecreateInfo != nullptr && RecreateInfo->SwapChain != VK_NULL_HANDLE)
 	{
@@ -172,16 +179,41 @@ FVulkanSwapChain::FVulkanSwapChain(VkInstance InInstance, FVulkanDevice& InDevic
 	else
 	{
 		// let the platform create the surface
-		FVulkanPlatform::CreateSurface(WindowContext, Instance, &Surface);
+		FVulkanPlatform::CreateSurface(WindowContext, InInstance, &Surface);
+	}
+
+	ON_SCOPE_EXIT
+	{
+		if (RecreateInfo != nullptr)
+		{
+			if (RecreateInfo->SwapChain != VK_NULL_HANDLE)
+			{
+				FVulkanPlatform::DestroySwapchainKHR(Device.GetHandle(), RecreateInfo->SwapChain, VULKAN_CPU_ALLOCATOR);
+				RecreateInfo->SwapChain = VK_NULL_HANDLE;
+			}
+			if (RecreateInfo->Surface != VK_NULL_HANDLE)
+			{
+				VulkanRHI::vkDestroySurfaceKHR(InInstance, RecreateInfo->Surface, VULKAN_CPU_ALLOCATOR);
+				RecreateInfo->Surface = VK_NULL_HANDLE;
+			}
+		}
+	};
+
+#if !UE_BUILD_SHIPPING
+	if (GSimulateLostSurfaceInNextSwapChain)
+	{
+		GSimulateLostSurfaceInNextSwapChain = false;
+		VERIFYVULKANSURFACERESULT(VK_ERROR_SURFACE_LOST_KHR);
 	}
+#endif
 
 	uint32 NumFormats;
-	VERIFYVULKANRESULT_EXPANDED(VulkanRHI::vkGetPhysicalDeviceSurfaceFormatsKHR(Device.GetPhysicalHandle(), Surface, &NumFormats, nullptr));
+	VERIFYVULKANSURFACERESULT(VulkanRHI::vkGetPhysicalDeviceSurfaceFormatsKHR(Device.GetPhysicalHandle(), Surface, &NumFormats, nullptr));
 	check(NumFormats > 0);
 
 	TArray<VkSurfaceFormatKHR> Formats;
 	Formats.AddZeroed(NumFormats);
-	VERIFYVULKANRESULT_EXPANDED(VulkanRHI::vkGetPhysicalDeviceSurfaceFormatsKHR(Device.GetPhysicalHandle(), Surface, &NumFormats, Formats.GetData()));
+	VERIFYVULKANSURFACERESULT(VulkanRHI::vkGetPhysicalDeviceSurfaceFormatsKHR(Device.GetPhysicalHandle(), Surface, &NumFormats, Formats.GetData()));
 
 	VkColorSpaceKHR RequestedColorSpace = Formats[0].colorSpace;
 
@@ -345,12 +377,12 @@ FVulkanSwapChain::FVulkanSwapChain(VkInstance InInstance, FVulkanDevice& InDevic
 		static bool bFirstTimeLog = !!VULKAN_HAS_DEBUGGING_ENABLED;
 
 		uint32 NumFoundPresentModes = 0;
-		VERIFYVULKANRESULT(VulkanRHI::vkGetPhysicalDeviceSurfacePresentModesKHR(Device.GetPhysicalHandle(), Surface, &NumFoundPresentModes, nullptr));
+		VERIFYVULKANSURFACERESULT(VulkanRHI::vkGetPhysicalDeviceSurfacePresentModesKHR(Device.GetPhysicalHandle(), Surface, &NumFoundPresentModes, nullptr));
 		check(NumFoundPresentModes > 0);
 
 		TArray<VkPresentModeKHR> FoundPresentModes;
 		FoundPresentModes.AddZeroed(NumFoundPresentModes);
-		VERIFYVULKANRESULT(VulkanRHI::vkGetPhysicalDeviceSurfacePresentModesKHR(Device.GetPhysicalHandle(), Surface, &NumFoundPresentModes, FoundPresentModes.GetData()));
+		VERIFYVULKANSURFACERESULT(VulkanRHI::vkGetPhysicalDeviceSurfacePresentModesKHR(Device.GetPhysicalHandle(), Surface, &NumFoundPresentModes, FoundPresentModes.GetData()));
 
 		UE_CLOG(bFirstTimeLog, LogVulkanRHI, Display, TEXT("Found %d Surface present modes:"), NumFoundPresentModes);
 
@@ -452,18 +484,13 @@ FVulkanSwapChain::FVulkanSwapChain(VkInstance InInstance, FVulkanDevice& InDevic
 	// Check the surface properties and formats
 
 	VkSurfaceCapabilitiesKHR SurfProperties;
-	VERIFYVULKANRESULT_EXPANDED(VulkanRHI::vkGetPhysicalDeviceSurfaceCapabilitiesKHR(Device.GetPhysicalHandle(),
-		Surface,
-		&SurfProperties));
+	VERIFYVULKANSURFACERESULT(VulkanRHI::vkGetPhysicalDeviceSurfaceCapabilitiesKHR(Device.GetPhysicalHandle(), Surface, &SurfProperties));
 
-	if (!IsVulkanPreTransformEnabled(GMaxRHIShaderPlatform))
-	{
-		PreTransform = VK_SURFACE_TRANSFORM_IDENTITY_BIT_KHR;
-	}
-	else
-	{
-		PreTransform = SurfProperties.currentTransform;
-	}
+	VkBool32 bSupportsPresent;
+	VERIFYVULKANSURFACERESULT(VulkanRHI::vkGetPhysicalDeviceSurfaceSupportKHR(Device.GetPhysicalHandle(), Device.GetPresentQueue()->GetFamilyIndex(), Surface, &bSupportsPresent));
+	ensure(bSupportsPresent);
+
+	const VkSurfaceTransformFlagBitsKHR PreTransform = IsVulkanPreTransformEnabled(GMaxRHIShaderPlatform) ? SurfProperties.currentTransform : VK_SURFACE_TRANSFORM_IDENTITY_BIT_KHR;
 
 	VkCompositeAlphaFlagBitsKHR CompositeAlpha = VK_COMPOSITE_ALPHA_INHERIT_BIT_KHR;
 	if (SurfProperties.supportedCompositeAlpha & VK_COMPOSITE_ALPHA_OPAQUE_BIT_KHR)
@@ -472,10 +499,10 @@ FVulkanSwapChain::FVulkanSwapChain(VkInstance InInstance, FVulkanDevice& InDevic
 	}
 
 	// 0 means no limit, so use the requested number
-	uint32 DesiredNumBuffers = SurfProperties.maxImageCount > 0 ? FMath::Clamp(*InOutDesiredNumBackBuffers, SurfProperties.minImageCount, SurfProperties.maxImageCount) : *InOutDesiredNumBackBuffers;
+	const uint32 DesiredNumBuffers = SurfProperties.maxImageCount > 0 ? FMath::Clamp(*InOutDesiredNumBackBuffers, SurfProperties.minImageCount, SurfProperties.maxImageCount) : *InOutDesiredNumBackBuffers;
 
-	uint32 SizeX = FVulkanPlatform::SupportsQuerySurfaceProperties() ? (SurfProperties.currentExtent.width == 0xFFFFFFFF ? Width : SurfProperties.currentExtent.width) : Width;
-	uint32 SizeY = FVulkanPlatform::SupportsQuerySurfaceProperties() ? (SurfProperties.currentExtent.height == 0xFFFFFFFF ? Height : SurfProperties.currentExtent.height) : Height;
+	const uint32 SizeX = FVulkanPlatform::SupportsQuerySurfaceProperties() ? (SurfProperties.currentExtent.width == 0xFFFFFFFF ? Width : SurfProperties.currentExtent.width) : Width;
+	const uint32 SizeY = FVulkanPlatform::SupportsQuerySurfaceProperties() ? (SurfProperties.currentExtent.height == 0xFFFFFFFF ? Height : SurfProperties.currentExtent.height) : Height;
 
 	VkSwapchainCreateInfoKHR SwapChainInfo;
 	ZeroVulkanStruct(SwapChainInfo, VK_STRUCTURE_TYPE_SWAPCHAIN_CREATE_INFO_KHR);
@@ -491,14 +518,14 @@ FVulkanSwapChain::FVulkanSwapChain(VkInstance InInstance, FVulkanDevice& InDevic
 	SwapChainInfo.imageSharingMode = VK_SHARING_MODE_EXCLUSIVE;
 	SwapChainInfo.presentMode = PresentMode;
 	SwapChainInfo.oldSwapchain = VK_NULL_HANDLE;
-	if(RecreateInfo != nullptr)
+	if (RecreateInfo != nullptr)
 	{
 		SwapChainInfo.oldSwapchain = RecreateInfo->SwapChain;
 	}
-	
+
 	SwapChainInfo.clipped = VK_TRUE;
 	SwapChainInfo.compositeAlpha = CompositeAlpha;
-	
+
 	{
 		//#todo-rco: Crappy workaround
 		if (SwapChainInfo.imageExtent.width == 0)
@@ -511,17 +538,14 @@ FVulkanSwapChain::FVulkanSwapChain(VkInstance InInstance, FVulkanDevice& InDevic
 		}
 	}
 
-	VkBool32 bSupportsPresent;
-	VERIFYVULKANRESULT(VulkanRHI::vkGetPhysicalDeviceSurfaceSupportKHR(Device.GetPhysicalHandle(), Device.GetPresentQueue()->GetFamilyIndex(), Surface, &bSupportsPresent));
-	ensure(bSupportsPresent);
 
 	//ensure(SwapChainInfo.imageExtent.width >= SurfProperties.minImageExtent.width && SwapChainInfo.imageExtent.width <= SurfProperties.maxImageExtent.width);
 	//ensure(SwapChainInfo.imageExtent.height >= SurfProperties.minImageExtent.height && SwapChainInfo.imageExtent.height <= SurfProperties.maxImageExtent.height);
 	static bool bPrintSwapchainCreationInfo = true;
 	if (bPrintSwapchainCreationInfo)
 	{
-		UE_LOG(LogVulkanRHI, Log, TEXT("Creating new VK swapchain with %s, %s, %s, num images %d"), 
-			VK_TYPE_TO_STRING(VkPresentModeKHR, SwapChainInfo.presentMode), VK_TYPE_TO_STRING(VkFormat, SwapChainInfo.imageFormat), 
+		UE_LOG(LogVulkanRHI, Log, TEXT("Creating new VK swapchain with %s, %s, %s, num images %d"),
+			VK_TYPE_TO_STRING(VkPresentModeKHR, SwapChainInfo.presentMode), VK_TYPE_TO_STRING(VkFormat, SwapChainInfo.imageFormat),
 			VK_TYPE_TO_STRING(VkColorSpaceKHR, SwapChainInfo.imageColorSpace), static_cast<uint32>(SwapChainInfo.minImageCount));
 #if WITH_EDITOR
 		bPrintSwapchainCreationInfo = false;
@@ -545,6 +569,7 @@ FVulkanSwapChain::FVulkanSwapChain(VkInstance InInstance, FVulkanDevice& InDevic
 		Swap(SwapChainInfo.imageExtent.width, SwapChainInfo.imageExtent.height);
 	}
 
+	VkSwapchainKHR SwapChain = VK_NULL_HANDLE;
 	VkResult Result = FVulkanPlatform::CreateSwapchainKHR(WindowContext, Device.GetPhysicalHandle(), Device.GetHandle(), &SwapChainInfo, VULKAN_CPU_ALLOCATOR, &SwapChain);
 
 	// Swap the aspect ratio back to make it transparent to high level
@@ -562,25 +587,10 @@ FVulkanSwapChain::FVulkanSwapChain(VkInstance InInstance, FVulkanDevice& InDevic
 		Result = FVulkanPlatform::CreateSwapchainKHR(WindowContext, Device.GetPhysicalHandle(), Device.GetHandle(), &SwapChainInfo, VULKAN_CPU_ALLOCATOR, &SwapChain);
 	}
 #endif
-	VERIFYVULKANRESULT_EXPANDED(Result);
-
-	if (RecreateInfo != nullptr)
-	{
-		if (RecreateInfo->SwapChain != VK_NULL_HANDLE)
-		{
-			FVulkanPlatform::DestroySwapchainKHR(Device.GetHandle(), RecreateInfo->SwapChain, VULKAN_CPU_ALLOCATOR);
-			RecreateInfo->SwapChain = VK_NULL_HANDLE;
-		}
-		if (RecreateInfo->Surface != VK_NULL_HANDLE)
-		{
-			VulkanRHI::vkDestroySurfaceKHR(Instance, RecreateInfo->Surface, VULKAN_CPU_ALLOCATOR);
-			RecreateInfo->Surface = VK_NULL_HANDLE;
-		}
-	}
+	VERIFYVULKANSURFACERESULT(Result);
 
-	InternalWidth = FMath::Min(Width, SwapChainInfo.imageExtent.width);
-	InternalHeight = FMath::Min(Height, SwapChainInfo.imageExtent.height);
-	bInternalFullScreen = bIsFullScreen;
+	Width = FMath::Min(Width, SwapChainInfo.imageExtent.width);
+	Height = FMath::Min(Height, SwapChainInfo.imageExtent.height);
 
 	uint32 NumSwapChainImages;
 	VERIFYVULKANRESULT_EXPANDED(VulkanRHI::vkGetSwapchainImagesKHR(Device.GetHandle(), SwapChain, &NumSwapChainImages, nullptr));
@@ -593,6 +603,29 @@ FVulkanSwapChain::FVulkanSwapChain(VkInstance InInstance, FVulkanDevice& InDevic
 	OutImages.AddUninitialized(NumSwapChainImages);
 	VERIFYVULKANRESULT_EXPANDED(VulkanRHI::vkGetSwapchainImagesKHR(Device.GetHandle(), SwapChain, &NumSwapChainImages, OutImages.GetData()));
 
+	*InOutDesiredNumBackBuffers = NumSwapChainImages;
+
+	return new FVulkanSwapChain(InInstance, Device, Surface, SwapChain, PreTransform, Width, Height, bIsFullScreen, NumSwapChainImages, LockToVsync, WindowContext);
+#undef VERIFYVULKANSURFACERESULT
+}
+
+FVulkanSwapChain::FVulkanSwapChain(
+	VkInstance InInstance, FVulkanDevice& InDevice, VkSurfaceKHR InSurface,	VkSwapchainKHR InSwapChain, 
+	VkSurfaceTransformFlagBitsKHR InPreTransform, uint32 InWidth, uint32 InHeight, bool InIsFullscreen,
+	uint32 NumSwapChainImages, int8 InLockToVsync, FVulkanPlatformWindowContext& WindowContext)
+	: Instance(InInstance)
+	, Device(InDevice)
+	, Surface(InSurface)
+	, SwapChain(InSwapChain)
+	, PreTransform(InPreTransform)
+	, InternalWidth(InWidth)
+	, InternalHeight(InHeight)
+	, bInternalFullScreen(InIsFullscreen)
+	, WindowHandle(WindowContext.GetWindowHandle())
+	, LockToVsync(InLockToVsync)
+{
+	NextPresentTargetTime = (FPlatformTime::Seconds() - GStartTime);
+
 #if VULKAN_USE_IMAGE_ACQUIRE_FENCES
 	ImageAcquiredFences.AddUninitialized(NumSwapChainImages);
 	FVulkanFenceManager& FenceMgr = Device.GetFenceManager();
@@ -607,10 +640,6 @@ FVulkanSwapChain::FVulkanSwapChain(VkInstance InInstance, FVulkanDevice& InDevic
 		ImageAcquiredSemaphore[BufferIndex] = new FVulkanSemaphore(Device, EVulkanSemaphoreFlags::ExternallySignaled);
 		ImageAcquiredSemaphore[BufferIndex]->AddRef();
 	}
-
-	*InOutDesiredNumBackBuffers = NumSwapChainImages;
-
-	PresentID = 0;
 }
 
 void FVulkanSwapChain::Destroy(FVulkanSwapChainRecreateInfo* RecreateInfo)
@@ -645,7 +674,7 @@ void FVulkanSwapChain::Destroy(FVulkanSwapChainRecreateInfo* RecreateInfo)
 		ImageAcquiredSemaphore[BufferIndex]->Release();
 	}
 
-	if(!bRecreate)
+	if (!bRecreate)
 	{
 		VulkanRHI::vkDestroySurfaceKHR(Instance, Surface, VULKAN_CPU_ALLOCATOR);
 	}
diff --git a/Engine/Source/Runtime/VulkanRHI/Private/VulkanSwapChain.h b/Engine/Source/Runtime/VulkanRHI/Private/VulkanSwapChain.h
--- a/Engine/Source/Runtime/VulkanRHI/Private/VulkanSwapChain.h
+++ b/Engine/Source/Runtime/VulkanRHI/Private/VulkanSwapChain.h
@@ -28,9 +28,16 @@ struct FVulkanSwapChainRecreateInfo
 
 class FVulkanSwapChain
 {
+	FVulkanSwapChain(VkInstance InInstance, FVulkanDevice& InDevice, VkSurfaceKHR InSurface, VkSwapchainKHR InSwapChain, 
+		VkSurfaceTransformFlagBitsKHR PreTransform, uint32 InWidth, uint32 InHeight, bool InIsFullscreen, uint32 NumSwapChainImages,
+		int8 InLockToVsync, FVulkanPlatformWindowContext& WindowContext);
+
 public:
-	FVulkanSwapChain(VkInstance InInstance, FVulkanDevice& InDevice, EPixelFormat& InOutPixelFormat, uint32 Width, uint32 Height, bool bIsFullscreen,
-		uint32* InOutDesiredNumBackBuffers, TArray<VkImage>& OutImages, int8 bLockToVsync, FVulkanPlatformWindowContext& PlatformContext, FVulkanSwapChainRecreateInfo* RecreateInfo);
+	static FVulkanSwapChain* Create(
+		VkInstance InInstance, FVulkanDevice& Device, EPixelFormat& InOutPixelFormat,
+		uint32 Width, uint32 Height, bool bIsFullScreen,
+		uint32* InOutDesiredNumBackBuffers, TArray<VkImage>& OutImages, int8 LockToVsync,
+		FVulkanPlatformWindowContext& WindowContext, FVulkanSwapChainRecreateInfo* RecreateInfo);
 
 	void Destroy(FVulkanSwapChainRecreateInfo* RecreateInfo);
 
@@ -49,35 +56,36 @@ public:
 	inline VkSurfaceTransformFlagBitsKHR GetCachedSurfaceTransform() const { return PreTransform; }
 
 protected:
-	VkSurfaceTransformFlagBitsKHR PreTransform;
-	VkFormat ImageFormat = VK_FORMAT_UNDEFINED;
-
-	VkSwapchainKHR SwapChain;
+	const VkInstance Instance;
 	FVulkanDevice& Device;
 
-	VkSurfaceKHR Surface;
-	void* WindowHandle;
-		
-	int32 CurrentImageIndex;
-	int32 SemaphoreIndex;
-	uint32 NumPresentCalls;
-	uint32 NumAcquireCalls;
+	VkSurfaceKHR Surface = VK_NULL_HANDLE;
+	VkSwapchainKHR SwapChain = VK_NULL_HANDLE;
+
+	const VkSurfaceTransformFlagBitsKHR PreTransform;
+
 	uint32 InternalWidth = 0;
 	uint32 InternalHeight = 0;
 	bool bInternalFullScreen = false;
 
+	void* WindowHandle = nullptr;
+		
+	int32 CurrentImageIndex = -1;
+	int32 SemaphoreIndex = 0;
+	uint32 NumPresentCalls = 0;
+	uint32 NumAcquireCalls = 0;
+
 	uint32 RTPacingSampleCount = 0;
 	double RTPacingPreviousFrameCPUTime = 0;
 	double RTPacingSampledDeltaTimeMS = 0;
 	
 	double NextPresentTargetTime = 0;
 
-	VkInstance Instance;
 	TArray<FVulkanSemaphore*> ImageAcquiredSemaphore;
 #if VULKAN_USE_IMAGE_ACQUIRE_FENCES
 	TArray<FVulkanFence*> ImageAcquiredFences;
 #endif
-	int8 LockToVsync;
+	int8 LockToVsync = 0;
 
 	uint32 PresentID = 0;
 
diff --git a/Engine/Source/Runtime/VulkanRHI/Private/VulkanViewport.cpp b/Engine/Source/Runtime/VulkanRHI/Private/VulkanViewport.cpp
--- a/Engine/Source/Runtime/VulkanRHI/Private/VulkanViewport.cpp
+++ b/Engine/Source/Runtime/VulkanRHI/Private/VulkanViewport.cpp
@@ -176,10 +176,26 @@ FVulkanViewport::~FVulkanViewport()
 	FVulkanDynamicRHI::Get().Viewports.Remove(this);
 }
 
-bool FVulkanViewport::DoCheckedSwapChainJob(FVulkanCommandListContext& Context, TFunction<int32(FVulkanViewport*)> SwapChainJob)
+bool FVulkanViewport::DoCheckedSwapChainJob(FVulkanCommandListContext& Context)
 {
+	FVulkanQueue* PresentQueue = Device.GetPresentQueue();
+
+	// Present the back buffer to the viewport window.
+	auto SwapChainJob = [PresentQueue, Viewport=this]()
+	{
+		// May happend if swapchain was recreated in DoCheckedSwapChainJob()
+		if (Viewport->AcquiredImageIndex == -1)
+		{
+			// Skip present silently if image has not been acquired
+			return (int32)FVulkanSwapChain::EStatus::Healthy;
+		}
+
+		return (int32)Viewport->SwapChain->Present(PresentQueue, Viewport->RenderingDoneSemaphores[Viewport->AcquiredImageIndex]);
+	};
+
+
 	int32 AttemptsPending = FVulkanPlatform::RecreateSwapchainOnFail() ? 4 : 0;
-	int32 Status = SwapChainJob(this);
+	int32 Status = SwapChainJob();
 
 	while (FVulkanPlatformWindowContext::CanCreateSwapchainOnDemand() && Status < 0 && AttemptsPending > 0)
 	{
@@ -202,7 +218,7 @@ bool FVulkanViewport::DoCheckedSwapChainJob(FVulkanCommandListContext& Context,
 		// Swapchain creation pushes some commands - flush the command buffers now to begin with a fresh state
 		Context.FlushCommands(EVulkanFlushFlags::WaitForCompletion);
 
-		Status = SwapChainJob(this);
+		Status = SwapChainJob();
 
 		--AttemptsPending;
 	}
@@ -417,7 +433,7 @@ void FVulkanViewport::CreateSwapchain(FVulkanCommandListContext& Context, FVulka
 		{
 			uint32 DesiredNumBackBuffers = NUM_BUFFERS;
 			TArray<VkImage> Images;
-			SwapChain = new FVulkanSwapChain(
+			SwapChain = FVulkanSwapChain::Create(
 				FVulkanDynamicRHI::Get().Instance, Device,
 				PixelFormat, SizeX, SizeY, bIsFullscreen,
 				&DesiredNumBackBuffers,
@@ -427,23 +443,26 @@ void FVulkanViewport::CreateSwapchain(FVulkanCommandListContext& Context, FVulka
 				RecreateInfo
 			);
 
-			checkf(Images.Num() >= NUM_BUFFERS, TEXT("We wanted at least %i images, actual Num: %i"), NUM_BUFFERS, Images.Num());
+			if (SwapChain)
+			{
+				checkf(Images.Num() >= NUM_BUFFERS, TEXT("We wanted at least %i images, actual Num: %i"), NUM_BUFFERS, Images.Num());
 
-			const bool bCreateSemaphores = RenderingDoneSemaphores.IsEmpty();
-			checkf(bCreateSemaphores || RenderingDoneSemaphores.Num() == Images.Num(), 
-				TEXT("CreateSwapchain, image count is not expected to change:  Requested=%d  Actual=%d  Previous=%d"),
-				NUM_BUFFERS, Images.Num(), RenderingDoneSemaphores.Num());
+				const bool bCreateSemaphores = RenderingDoneSemaphores.IsEmpty();
+				checkf(bCreateSemaphores || RenderingDoneSemaphores.Num() == Images.Num(),
+					TEXT("CreateSwapchain, image count is not expected to change:  Requested=%d  Actual=%d  Previous=%d"),
+					NUM_BUFFERS, Images.Num(), RenderingDoneSemaphores.Num());
 
-			BackBufferImages.SetNum(Images.Num());
-			RenderingDoneSemaphores.SetNum(Images.Num());
-			InitImages(Context, Images);
+				BackBufferImages.SetNum(Images.Num());
+				RenderingDoneSemaphores.SetNum(Images.Num());
+				InitImages(Context, Images);
 
-			if (bCreateSemaphores)
-			{
-				for (int32 Index = 0, NumBuffers = RenderingDoneSemaphores.Num(); Index < NumBuffers; ++Index)
+				if (bCreateSemaphores)
 				{
-					RenderingDoneSemaphores[Index] = new FVulkanSemaphore(Device);
-					RenderingDoneSemaphores[Index]->AddRef();
+					for (int32 Index = 0, NumBuffers = RenderingDoneSemaphores.Num(); Index < NumBuffers; ++Index)
+					{
+						RenderingDoneSemaphores[Index] = new FVulkanSemaphore(Device);
+						RenderingDoneSemaphores[Index]->AddRef();
+					}
 				}
 			}
 		}
@@ -470,8 +489,8 @@ void FVulkanViewport::CreateSwapchain(FVulkanCommandListContext& Context, FVulka
 
 	// We always create a 'dummy' backbuffer to gracefully handle SurfaceLost cases
 	{
-		uint32 BackBufferSizeX = RequiresRenderingBackBuffer() ? SizeX : 1;
-		uint32 BackBufferSizeY = RequiresRenderingBackBuffer() ? SizeY : 1;
+		const uint32 BackBufferSizeX = RequiresRenderingBackBuffer() ? SizeX : 1;
+		const uint32 BackBufferSizeY = RequiresRenderingBackBuffer() ? SizeY : 1;
 
 		const UE::StereoRenderUtils::FStereoShaderAspects Aspects(GMaxRHIShaderPlatform);
 
@@ -487,9 +506,7 @@ void FVulkanViewport::CreateSwapchain(FVulkanCommandListContext& Context, FVulka
 		RenderingBackBuffer = FVulkanDynamicRHI::Get().CreateTextureInternal(CreateDesc, {});
 
 		FVulkanPipelineBarrier Barrier;
-
 		Barrier.AddImageLayoutTransition(RenderingBackBuffer->Image, VK_IMAGE_LAYOUT_UNDEFINED, VK_IMAGE_LAYOUT_PRESENT_SRC_KHR, FVulkanPipelineBarrier::MakeSubresourceRange(VK_IMAGE_ASPECT_COLOR_BIT, 0, 1));
-
 		Barrier.Execute(&Context.GetCommandBuffer());
 
 #if VULKAN_ENABLE_DRAW_MARKERS
@@ -652,7 +669,7 @@ inline static void CopyImageToBackBuffer(FVulkanCommandListContext& Context, FVu
 	}
 }
 
-bool FVulkanViewport::Present(FVulkanCommandListContext& Context, FVulkanQueue* PresentQueue, bool bLockToVsync)
+bool FVulkanViewport::Present(FVulkanCommandListContext& Context, bool bLockToVsync)
 {
 	check(Context.IsImmediate());
 
@@ -707,7 +724,7 @@ bool FVulkanViewport::Present(FVulkanCommandListContext& Context, FVulkanQueue*
 		}
 		else
 		{
-			if(FVulkanPlatformWindowContext::CanCreateSwapchainOnDemand()) // on android we dont want to attempt to recreate the window when we dont have the window lock...
+			if (FVulkanPlatformWindowContext::CanCreateSwapchainOnDemand()) // on android we dont want to attempt to recreate the window when we dont have the window lock...
 			{
 				// failing to do the delayacquire can only happen if we were in this mode to begin with
 				check(GVulkanDelayAcquireImage == EDelayAcquireImageType::DelayAcquire);
@@ -743,19 +760,7 @@ bool FVulkanViewport::Present(FVulkanCommandListContext& Context, FVulkanQueue*
 	bool bResult = false;
 	if (bNeedNativePresent && (!SupportsStandardSwapchain() || GVulkanDelayAcquireImage == EDelayAcquireImageType::DelayAcquire || RHIBackBuffer != nullptr))
 	{
-		// Present the back buffer to the viewport window.
-		auto SwapChainJob = [PresentQueue](FVulkanViewport* Viewport)
-		{
-			// May happend if swapchain was recreated in DoCheckedSwapChainJob()
-			if (Viewport->AcquiredImageIndex == -1)
-			{
-				// Skip present silently if image has not been acquired
-				return (int32)FVulkanSwapChain::EStatus::Healthy;
-			}
-		
-			return (int32)Viewport->SwapChain->Present(PresentQueue, Viewport->RenderingDoneSemaphores[Viewport->AcquiredImageIndex]);
-		};
-		if (SupportsStandardSwapchain() && !DoCheckedSwapChainJob(Context, SwapChainJob))
+		if (SupportsStandardSwapchain() && !DoCheckedSwapChainJob(Context))
 		{
 			UE_LOG(LogVulkanRHI, Error, TEXT("Swapchain present failed!"));
 			bResult = false;
@@ -796,11 +801,6 @@ bool FVulkanViewport::Present(FVulkanCommandListContext& Context, FVulkanQueue*
 	return bResult;
 }
 
-VkFormat FVulkanViewport::GetSwapchainImageFormat() const
-{
-	return SwapChain->ImageFormat;
-}
-
 bool FVulkanViewport::SupportsStandardSwapchain()
 {
 	return !bRenderOffscreen && !FVulkanDynamicRHI::Get().bIsStandaloneStereoDevice;
diff --git a/Engine/Source/Runtime/VulkanRHI/Private/VulkanViewport.h b/Engine/Source/Runtime/VulkanRHI/Private/VulkanViewport.h
--- a/Engine/Source/Runtime/VulkanRHI/Private/VulkanViewport.h
+++ b/Engine/Source/Runtime/VulkanRHI/Private/VulkanViewport.h
@@ -72,7 +72,7 @@ public:
 
 	virtual void Tick(float DeltaTime) override final;
 	
-	bool Present(FVulkanCommandListContext& Context, FVulkanQueue* PresentQueue, bool bLockToVsync);
+	bool Present(FVulkanCommandListContext& Context, bool bLockToVsync);
 
 	inline uint32 GetPresentCount() const
 	{
@@ -106,7 +106,6 @@ public:
 		return SwapChain;
 	}
 
-	VkFormat GetSwapchainImageFormat() const;
 	inline void* GetWindowHandle() { return WindowHandle; }
 
 	void DestroySwapchain(struct FVulkanSwapChainRecreateInfo* RecreateInfo);
@@ -160,7 +159,7 @@ protected:
 	}
 	void Resize(FRHICommandListImmediate& RHICmdList, uint32 InSizeX, uint32 InSizeY, bool bIsFullscreen, EPixelFormat PreferredPixelFormat, FVulkanPlatformWindowContext& WindowContext);
 
-	bool DoCheckedSwapChainJob(FVulkanCommandListContext& Context, TFunction<int32(FVulkanViewport*)> SwapChainJob);
+	bool DoCheckedSwapChainJob(FVulkanCommandListContext& Context);
 	bool SupportsStandardSwapchain();
 	bool RequiresRenderingBackBuffer();
 	EPixelFormat GetPixelFormatForNonDefaultSwapchain();
