VulkanRHI: Add null check on Swapchain for viewport destructor

#rnx

diff --git a/Engine/Source/Runtime/VulkanRHI/Private/VulkanViewport.cpp b/Engine/Source/Runtime/VulkanRHI/Private/VulkanViewport.cpp
--- a/Engine/Source/Runtime/VulkanRHI/Private/VulkanViewport.cpp
+++ b/Engine/Source/Runtime/VulkanRHI/Private/VulkanViewport.cpp
@@ -166,9 +166,12 @@ FVulkanViewport::~FVulkanViewport()
 			BackBufferImages[Index] = nullptr;
 		}
 
-		SwapChain->Destroy(nullptr);
-		delete SwapChain;
-		SwapChain = nullptr;
+		if (SwapChain)
+		{
+			SwapChain->Destroy(nullptr);
+			delete SwapChain;
+			SwapChain = nullptr;
+		}
 
 		FCoreDelegates::OnSystemResolutionChanged.RemoveAll(this);
 	}
