VulkanRHI: Merge consecutive tile commits for sparse bindings.

#rnx

diff --git a/Engine/Source/Runtime/VulkanRHI/Private/VulkanTexture.cpp b/Engine/Source/Runtime/VulkanRHI/Private/VulkanTexture.cpp
--- a/Engine/Source/Runtime/VulkanRHI/Private/VulkanTexture.cpp
+++ b/Engine/Source/Runtime/VulkanRHI/Private/VulkanTexture.cpp
@@ -47,7 +47,12 @@ static TAutoConsoleVariable<int32> CVarVulkanSparseImageAllocSizeMB(
 	ECVF_ReadOnly
 );
 
-
+static TAutoConsoleVariable<int32> CVarVulkanSparseImageMergeBindings(
+	TEXT("r.Vulkan.SparseImageMergeBindings"),
+	1,
+	TEXT("Merge consecutive bindings of sparse images."),
+	ECVF_ReadOnly
+);
 
 extern int32 GVulkanLogDefrag;
 
@@ -911,6 +916,8 @@ TArray<VkSparseMemoryBind> FVulkanTexture::CommitReservedResource(uint64 Require
 	checkf((RequiredCommitSizeInBytes == UINT64_MAX) || (RequiredCommitSizeInBytes == 0), TEXT("Only full resource commits are supported for images."));
 	checkf(ReservedResourceData, TEXT("CommitReservedResource called on a resource without any ReservedResourceData.  Resource needs to be created with the ReservedResource flag."));
 
+	const bool bMergeConsecutiveBindings = CVarVulkanSparseImageMergeBindings.GetValueOnAnyThread() != 0;
+
 	const uint64 SparseBlockSize = (uint64)CVarVulkanSparseImageAllocSizeMB.GetValueOnAnyThread() * 1024 * 1024;
 	const uint64 TileSize = GRHIGlobals.ReservedResources.TileSizeInBytes;
 	RequiredCommitSizeInBytes = Align(FMath::Min<uint64>(RequiredCommitSizeInBytes, ReservedResourceData->MemoryRequirements.size), TileSize);
@@ -947,6 +954,7 @@ TArray<VkSparseMemoryBind> FVulkanTexture::CommitReservedResource(uint64 Require
 
 		// Bind all the individual tiles
 		Binds.Reserve(RequiredTileCount - CurrentTileCount);
+		uint32 LastAllocIndex = MAX_uint32;
 		for (uint64 TileCount = CurrentTileCount; TileCount < RequiredTileCount; TileCount++)
 		{
 			const uint32 TileOffset = (TileCount * TileSize);
@@ -954,25 +962,45 @@ TArray<VkSparseMemoryBind> FVulkanTexture::CommitReservedResource(uint64 Require
 			const uint32 AllocIndex = TileOffset / SparseBlockSize;
 			const FVulkanReservedResourceData::FSparseAllocation& SparseAlloc = ReservedResourceData->Allocations[AllocIndex];
 
-			VkSparseMemoryBind& NewBind = Binds.AddZeroed_GetRef();
-			NewBind.size = TileSize;
-			NewBind.memory = SparseAlloc.Allocation.GetDeviceMemoryHandle(Device);
-			NewBind.memoryOffset = SparseAlloc.Allocation.Offset + TileOffset - SparseAlloc.ResourceOffset;
-			NewBind.resourceOffset = TileOffset;
+			if (bMergeConsecutiveBindings && !Binds.IsEmpty() && (LastAllocIndex == AllocIndex))
+			{
+				Binds.Last().size += TileSize;
+			}
+			else
+			{
+				VkSparseMemoryBind& NewBind = Binds.AddZeroed_GetRef();
+				NewBind.size = TileSize;
+				NewBind.memory = SparseAlloc.Allocation.GetDeviceMemoryHandle(Device);
+				NewBind.memoryOffset = SparseAlloc.Allocation.Offset + TileOffset - SparseAlloc.ResourceOffset;
+				NewBind.resourceOffset = TileOffset;
+			}
+
+			LastAllocIndex = AllocIndex;
 		}
 
-		const int64 CommitDeltaInBytes = TileSize * Binds.Num();
+		const int64 CommitDeltaInBytes = TileSize * (RequiredTileCount - CurrentTileCount);
 		UE::RHICore::UpdateReservedResourceStatsOnCommit(CommitDeltaInBytes, false /* bBuffer */, true /* bCommitting */);
 	}
 	else if (RequiredTileCount < CurrentTileCount)
 	{
-		for (int64 TileIndex = CurrentTileCount - 1; TileIndex >= (int64)RequiredTileCount; TileIndex--)
+		if (bMergeConsecutiveBindings)
 		{
 			VkSparseMemoryBind& NewBind = Binds.AddZeroed_GetRef();
-			NewBind.size = TileSize;
+			NewBind.size = (CurrentTileCount - RequiredTileCount) * TileSize;
 			NewBind.memory = VK_NULL_HANDLE;
 			NewBind.memoryOffset = 0;
-			NewBind.resourceOffset = (TileIndex * TileSize);
+			NewBind.resourceOffset = (RequiredTileCount * TileSize);
+		}
+		else
+		{
+			for (int64 TileIndex = CurrentTileCount - 1; TileIndex >= (int64)RequiredTileCount; TileIndex--)
+			{
+				VkSparseMemoryBind& NewBind = Binds.AddZeroed_GetRef();
+				NewBind.size = TileSize;
+				NewBind.memory = VK_NULL_HANDLE;
+				NewBind.memoryOffset = 0;
+				NewBind.resourceOffset = (TileIndex * TileSize);
+			}
 		}
 
 		// Free up unused blocks (deletion is deferred, don't have to wait for binds to be submitted)
@@ -982,7 +1010,7 @@ TArray<VkSparseMemoryBind> FVulkanTexture::CommitReservedResource(uint64 Require
 			ReservedResourceData->Allocations.Pop();
 		}
 
-		const int64 CommitDeltaInBytes = TileSize * Binds.Num();
+		const int64 CommitDeltaInBytes = TileSize * (CurrentTileCount - RequiredTileCount);
 		UE::RHICore::UpdateReservedResourceStatsOnCommit(CommitDeltaInBytes, false /* bBuffer */, false /* bCommitting */);
 	}
 
