VulkanRHI: Fix missing ref on syncpoint when signaled from top_of_pipe.

#rnx

diff --git a/Engine/Source/Runtime/VulkanRHI/Private/VulkanContext.h b/Engine/Source/Runtime/VulkanRHI/Private/VulkanContext.h
--- a/Engine/Source/Runtime/VulkanRHI/Private/VulkanContext.h
+++ b/Engine/Source/Runtime/VulkanRHI/Private/VulkanContext.h
@@ -166,7 +166,7 @@ public:
 	}
 
 	// Add an event to be signaled with the current payload is submitted (does not alter phase)
-	void AddSubmissionEvent(FGraphEventRef& InEvent)
+	void AddSubmissionEvent(const FGraphEventRef& InEvent)
 	{
 		FVulkanPayload& Payload = GetPayload(CurrentPhase);
 		Payload.SubmissionEvents.Add(InEvent);
diff --git a/Engine/Source/Runtime/VulkanRHI/Private/VulkanUtil.cpp b/Engine/Source/Runtime/VulkanRHI/Private/VulkanUtil.cpp
--- a/Engine/Source/Runtime/VulkanRHI/Private/VulkanUtil.cpp
+++ b/Engine/Source/Runtime/VulkanRHI/Private/VulkanUtil.cpp
@@ -701,22 +701,22 @@ void FVulkanDynamicRHI::RHIWriteGPUFence_TopOfPipe(FRHICommandListBase& RHICmdLi
 	}
 
 	Fence->NumPendingWriteCommands.Increment();
-	RHICmdList.EnqueueLambda([Fence](FRHICommandListBase& CmdList)
+	RHICmdList.EnqueueLambda([FenceRef=TRefCountPtr<FVulkanGPUFence>(Fence), CompletedSyncPoint=Fence->CompletedSyncPoint, SubmittedSyncPoint=Fence->SubmittedSyncPoint](FRHICommandListBase& CmdList)
 	{
 		FVulkanCommandListContext& Context = FVulkanCommandListContext::Get(CmdList);
 
-		if (Fence->Event)
+		if (FenceRef->Event)
 		{
-			VulkanRHI::vkCmdSetEvent(Context.GetCommandBuffer().GetHandle(), Fence->Event, VK_PIPELINE_STAGE_BOTTOM_OF_PIPE_BIT);
-			Context.AddSubmissionEvent(Fence->SubmittedSyncPoint);
-			Context.AddPendingSyncPoint(Fence->CompletedSyncPoint);
+			VulkanRHI::vkCmdSetEvent(Context.GetCommandBuffer().GetHandle(), FenceRef->Event, VK_PIPELINE_STAGE_BOTTOM_OF_PIPE_BIT);
+			Context.AddSubmissionEvent(SubmittedSyncPoint);
+			Context.AddPendingSyncPoint(CompletedSyncPoint);
 		}
 		else
 		{
-			Context.SignalSyncPoint(Fence->CompletedSyncPoint);
+			Context.SignalSyncPoint(CompletedSyncPoint);
 		}
 
-		Fence->NumPendingWriteCommands.Decrement();
+		FenceRef->NumPendingWriteCommands.Decrement();
 	});
 }
 
