Give vulkan platform code an opportunity to handle GPU crashes.
This restores Android's fatal logging on GPU failures and catches some new failure scenarios in submission.
Reinstate verbosity during vulkan crash logging.

diff --git a/Engine/Source/Runtime/VulkanRHI/Private/Android/VulkanAndroidPlatform.cpp b/Engine/Source/Runtime/VulkanRHI/Private/Android/VulkanAndroidPlatform.cpp
--- a/Engine/Source/Runtime/VulkanRHI/Private/Android/VulkanAndroidPlatform.cpp
+++ b/Engine/Source/Runtime/VulkanRHI/Private/Android/VulkanAndroidPlatform.cpp
@@ -2157,3 +2157,10 @@ FString FVulkanAndroidPlatform::GetVulkanProfileNameForFeatureLevel(ERHIFeatureL
 	}
 	return ProfileName;
 } 
+
+void FVulkanAndroidPlatform::OnGPUCrash(const TCHAR* Message)
+{
+	// trigger a fatal message which will make it visible to crash reporter.
+	GLog->Panic();
+	UE_LOG(LogVulkanRHI, Fatal, TEXT("%s"), Message);
+}
diff --git a/Engine/Source/Runtime/VulkanRHI/Private/Android/VulkanAndroidPlatform.h b/Engine/Source/Runtime/VulkanRHI/Private/Android/VulkanAndroidPlatform.h
--- a/Engine/Source/Runtime/VulkanRHI/Private/Android/VulkanAndroidPlatform.h
+++ b/Engine/Source/Runtime/VulkanRHI/Private/Android/VulkanAndroidPlatform.h
@@ -258,6 +258,8 @@ public:
 		return VK_SHADER_STAGE_FRAGMENT_BIT | VK_SHADER_STAGE_COMPUTE_BIT;
 	}
 
+	static void OnGPUCrash(const TCHAR* Message);
+
 	static VkTimeDomainKHR GetTimeDomain() { return VK_TIME_DOMAIN_CLOCK_MONOTONIC_KHR; }
 
 #if USE_ANDROID_VULKAN_SWAPPY
diff --git a/Engine/Source/Runtime/VulkanRHI/Private/VulkanDynamicRHI.h b/Engine/Source/Runtime/VulkanRHI/Private/VulkanDynamicRHI.h
--- a/Engine/Source/Runtime/VulkanRHI/Private/VulkanDynamicRHI.h
+++ b/Engine/Source/Runtime/VulkanRHI/Private/VulkanDynamicRHI.h
@@ -341,7 +341,7 @@ public:
 
 	// Called when the GPU has crashed. This function will not return.
 	void ProcessInterruptQueueOnGPUCrash();
-	void TerminateOnGPUCrash();
+	void TerminateOnGPUCrash(const TCHAR* Message);
 
 	bool HasInterruptThread() const
 	{
diff --git a/Engine/Source/Runtime/VulkanRHI/Private/VulkanGenericPlatform.h b/Engine/Source/Runtime/VulkanRHI/Private/VulkanGenericPlatform.h
--- a/Engine/Source/Runtime/VulkanRHI/Private/VulkanGenericPlatform.h
+++ b/Engine/Source/Runtime/VulkanRHI/Private/VulkanGenericPlatform.h
@@ -59,6 +59,9 @@ public:
 	static void GetDeviceExtensions(FVulkanDevice* Device, FVulkanDeviceExtensionArray& OutExtensions);
 	static void GetDeviceLayers(TArray<const ANSICHAR*>& OutLayers) {}
 
+	// Allows platform specific GPU crash handling
+	static void OnGPUCrash(const TCHAR* Message) {}
+
 	// create the platform-specific surface object - required
 	static void CreateSurface(VkSurfaceKHR* OutSurface);
 
diff --git a/Engine/Source/Runtime/VulkanRHI/Private/VulkanSubmission.cpp b/Engine/Source/Runtime/VulkanRHI/Private/VulkanSubmission.cpp
--- a/Engine/Source/Runtime/VulkanRHI/Private/VulkanSubmission.cpp
+++ b/Engine/Source/Runtime/VulkanRHI/Private/VulkanSubmission.cpp
@@ -494,7 +494,7 @@ bool FVulkanDynamicRHI::ProcessInterruptQueue()
 
 	if (GIsGPUCrashed)
 	{
-		TerminateOnGPUCrash();
+		TerminateOnGPUCrash(TEXT("ProcessInterruptQueue"));
 	}
 
 	return bProgress;
@@ -707,7 +707,7 @@ void FVulkanDynamicRHI::CompletePayload(FVulkanPayload* Payload)
 				{
 					// The GPU is crashed and we are the current interrupt thread, 
 					// terminate right away if we crashed since no more progress is possible
-					TerminateOnGPUCrash();
+					TerminateOnGPUCrash(TEXT("CompletePayload"));
 				}
 				else
 				{
diff --git a/Engine/Source/Runtime/VulkanRHI/Private/VulkanUtil.cpp b/Engine/Source/Runtime/VulkanRHI/Private/VulkanUtil.cpp
--- a/Engine/Source/Runtime/VulkanRHI/Private/VulkanUtil.cpp
+++ b/Engine/Source/Runtime/VulkanRHI/Private/VulkanUtil.cpp
@@ -795,7 +795,8 @@ namespace VulkanRHI
 			break;
 		}
 
-#define VULKAN_PRINT_ERROR_MSG(Verbosity) UE_LOG(LogVulkanRHI, Error, TEXT("%s failed, VkResult=%d\n at %s:%u \n with error %s"), ANSI_TO_TCHAR(VkFunction), (int32)Result, ANSI_TO_TCHAR(Filename), Line, *ErrorString)
+		FString Message = FString::Printf(TEXT("%s failed, VkResult=%d\n at %s:%u \n with error %s"), ANSI_TO_TCHAR(VkFunction), (int32)Result, ANSI_TO_TCHAR(Filename), Line, *ErrorString);
+#define VULKAN_PRINT_ERROR_MSG(Verbosity) UE_LOG(LogVulkanRHI, Verbosity, TEXT("%s"), *Message)
 
 		if (GIsGPUCrashed)
 		{
@@ -814,7 +815,7 @@ namespace VulkanRHI
 				}
 
 				// This function does not return.
-				GVulkanRHI->TerminateOnGPUCrash();
+				GVulkanRHI->TerminateOnGPUCrash(*Message);
 			}
 			else
 			{
@@ -1047,7 +1048,7 @@ void FVulkanDynamicRHI::DumpCrashMarkers()
 
 
 
-void FVulkanDynamicRHI::TerminateOnGPUCrash()
+void FVulkanDynamicRHI::TerminateOnGPUCrash(const TCHAR* Message)
 {
 	DumpCrashMarkers();
 
@@ -1069,7 +1070,9 @@ void FVulkanDynamicRHI::TerminateOnGPUCrash()
 		UE_DEBUG_BREAK();
 	}
 
-		// Force shutdown, we can't do anything useful anymore.
+	FVulkanPlatform::OnGPUCrash(Message);
+
+	// Force shutdown, we can't do anything useful anymore.
 	FPlatformMisc::RequestExit(true, TEXT("FVulkanDynamicRHI.TerminateOnGPUCrash"));
 }
 
