VulkanRHI:  Make it possible to set extension support based on its properties.  Add VK_EXT_host_image_copy functions to loader.

#rnx

diff --git a/Engine/Source/Runtime/VulkanRHI/Private/VulkanCommandWrappers.h b/Engine/Source/Runtime/VulkanRHI/Private/VulkanCommandWrappers.h
--- a/Engine/Source/Runtime/VulkanRHI/Private/VulkanCommandWrappers.h
+++ b/Engine/Source/Runtime/VulkanRHI/Private/VulkanCommandWrappers.h
@@ -295,6 +295,9 @@ struct FWrapLayer
 	static void CmdSetLogicOpEnableEXT(VkResult Result, VkCommandBuffer CommandBuffer, VkBool32 LogicOpEnable) VULKAN_LAYER_BODY
 	static void CmdSetAlphaToCoverageEnableEXT(VkResult Result, VkCommandBuffer CommandBuffer, VkBool32 AlphaToCoverageEnable) VULKAN_LAYER_BODY
 	static void CmdSetAlphaToOneEnableEXT(VkResult Result, VkCommandBuffer CommandBuffer, VkBool32 AlphaToOneEnable) VULKAN_LAYER_BODY
+
+	static void CopyMemoryToImageEXT(VkResult Result, VkDevice Device, const VkCopyMemoryToImageInfo* CopyMemoryToImageInfo) VULKAN_LAYER_BODY
+	static void TransitionImageLayoutEXT(VkResult Result, VkDevice Device, uint32_t TransitionCount, const VkHostImageLayoutTransitionInfo* Transitions) VULKAN_LAYER_BODY
 };
 
 
@@ -2077,6 +2080,22 @@ namespace VulkanRHI
 		FWrapLayer::CmdSetAlphaToOneEnableEXT(VK_SUCCESS, CommandBuffer, AlphaToOneEnable);
 	}
 
+	static FORCEINLINE_DEBUGGABLE VkResult vkCopyMemoryToImageEXT(VkDevice Device, const VkCopyMemoryToImageInfo* CopyMemoryToImageInfo)
+	{
+		FWrapLayer::CopyMemoryToImageEXT(VK_RESULT_MAX_ENUM, Device, CopyMemoryToImageInfo);
+		const VkResult Result = VULKANAPINAMESPACE::vkCopyMemoryToImageEXT(Device, CopyMemoryToImageInfo);
+		FWrapLayer::CopyMemoryToImageEXT(Result, Device, CopyMemoryToImageInfo);
+		return Result;
+	}
+
+	static FORCEINLINE_DEBUGGABLE VkResult vkTransitionImageLayoutEXT(VkDevice Device, uint32_t TransitionCount, const VkHostImageLayoutTransitionInfo* Transitions)
+	{
+		FWrapLayer::TransitionImageLayoutEXT(VK_RESULT_MAX_ENUM, Device, TransitionCount, Transitions);
+		const VkResult Result = VULKANAPINAMESPACE::vkTransitionImageLayoutEXT(Device, TransitionCount, Transitions);
+		FWrapLayer::TransitionImageLayoutEXT(Result, Device, TransitionCount, Transitions);
+		return Result;
+	}
+
 
 #if VULKAN_ENABLE_IMAGE_TRACKING_LAYER
 	void BindDebugLabelName(VkImage Image, const TCHAR* Name);
diff --git a/Engine/Source/Runtime/VulkanRHI/Private/VulkanDebug.cpp b/Engine/Source/Runtime/VulkanRHI/Private/VulkanDebug.cpp
--- a/Engine/Source/Runtime/VulkanRHI/Private/VulkanDebug.cpp
+++ b/Engine/Source/Runtime/VulkanRHI/Private/VulkanDebug.cpp
@@ -4995,6 +4995,26 @@ void FWrapLayer::CmdSetAlphaToOneEnableEXT(VkResult Result, VkCommandBuffer Comm
 	}
 }
 
+void FWrapLayer::CopyMemoryToImageEXT(VkResult Result, VkDevice Device, const VkCopyMemoryToImageInfo* CopyMemoryToImageInfo)
+{
+	if (Result == VK_RESULT_MAX_ENUM)
+	{
+#if VULKAN_ENABLE_DUMP_LAYER
+		PrintfBeginResult(FString::Printf(TEXT("CopyMemoryToImageEXT()")));
+#endif
+	}
+}
+
+void FWrapLayer::TransitionImageLayoutEXT(VkResult Result, VkDevice Device, uint32_t TransitionCount, const VkHostImageLayoutTransitionInfo* Transitions)
+{
+	if (Result == VK_RESULT_MAX_ENUM)
+	{
+#if VULKAN_ENABLE_DUMP_LAYER
+		PrintfBeginResult(FString::Printf(TEXT("TransitionImageLayoutEXT()")));
+#endif
+	}
+}
+
 
 
 #if VULKAN_ENABLE_IMAGE_TRACKING_LAYER
diff --git a/Engine/Source/Runtime/VulkanRHI/Private/VulkanDevice.cpp b/Engine/Source/Runtime/VulkanRHI/Private/VulkanDevice.cpp
--- a/Engine/Source/Runtime/VulkanRHI/Private/VulkanDevice.cpp
+++ b/Engine/Source/Runtime/VulkanRHI/Private/VulkanDevice.cpp
@@ -1344,7 +1344,7 @@ void FVulkanDevice::InitGPU()
 		{
 			if (UEExtension->InUse())
 			{
-				UEExtension->PostPhysicalDeviceProperties();
+				UEExtension->PostPhysicalDeviceProperties(OptionalDeviceExtensions);
 			}
 		}
 	}
diff --git a/Engine/Source/Runtime/VulkanRHI/Private/VulkanDevice.h b/Engine/Source/Runtime/VulkanRHI/Private/VulkanDevice.h
--- a/Engine/Source/Runtime/VulkanRHI/Private/VulkanDevice.h
+++ b/Engine/Source/Runtime/VulkanRHI/Private/VulkanDevice.h
@@ -110,6 +110,7 @@ struct FOptionalVulkanDeviceExtensions
 
 			// Promoted to 1.4
 			uint64 HasEXTLoadStoreOpNone : 1;
+			uint64 HasEXTHostImageCopy : 1;
 		};
 		uint64 Packed;
 	};
diff --git a/Engine/Source/Runtime/VulkanRHI/Private/VulkanExtensions.cpp b/Engine/Source/Runtime/VulkanRHI/Private/VulkanExtensions.cpp
--- a/Engine/Source/Runtime/VulkanRHI/Private/VulkanExtensions.cpp
+++ b/Engine/Source/Runtime/VulkanRHI/Private/VulkanExtensions.cpp
@@ -109,6 +109,15 @@ TAutoConsoleVariable<int32> GVulkanAllowGraphicPipelineLibraryCVar(
 	ECVF_ReadOnly
 );
 
+TAutoConsoleVariable<int32> GVulkanAllowHostImageCopyCVar(
+	TEXT("r.Vulkan.AllowHostImageCopy"),
+	0,
+	TEXT("Enables the use of VK_EXT_host_image_copy to avoid an extra copy on unified memory devices. \n")
+	TEXT("0: Do not enable support for Host Image Copy (default)\n")
+	TEXT("1: Enable Host Image Copy"),
+	ECVF_ReadOnly
+);
+
 TAutoConsoleVariable<int32> GVulkanVariableRateShadingFormatCVar(
 	TEXT("r.Vulkan.VRSFormat"),
 	0,
@@ -344,7 +353,7 @@ public:
 		AddToPNext(PhysicalDeviceProperties2, PhysicalDeviceDriverProperties);
 	}
 
-	virtual void PostPhysicalDeviceProperties() override final
+	virtual void PostPhysicalDeviceProperties(FOptionalVulkanDeviceExtensions& ExtensionFlags) override final
 	{
 		const VkPhysicalDeviceDriverPropertiesKHR& PhysicalDeviceDriverProperties = GetDeviceExtensionProperties().PhysicalDeviceDriverProperties;
 		UE_LOG(LogVulkanRHI, Display, TEXT("Vulkan Physical Device Driver Properties:"));
@@ -672,7 +681,7 @@ public:
 		}
 	}
 
-	virtual void PostPhysicalDeviceProperties() override final
+	virtual void PostPhysicalDeviceProperties(FOptionalVulkanDeviceExtensions& ExtensionFlags) override final
 	{
 		VkPhysicalDeviceFragmentShadingRateFeaturesKHR& FragmentShadingRateFeatures = GetDeviceExtensionProperties().FragmentShadingRateFeatures;
 		if (FragmentShadingRateFeatures.attachmentFragmentShadingRate == VK_TRUE)
@@ -751,7 +760,7 @@ public:
 		AddToPNext(PhysicalDeviceFeatures2, FragmentDensityMapFeatures);
 	}
 
-	virtual void PostPhysicalDeviceProperties() override final
+	virtual void PostPhysicalDeviceProperties(FOptionalVulkanDeviceExtensions& ExtensionFlags) override final
 	{
 		// Use the Fragment Density Map extension if the Fragment Shading Rate extension is not available or if both are available but 
 		// Fragment Density Map is the preferred user choice.
@@ -862,7 +871,7 @@ public:
 		GRHIGlobals.SupportsBarycentricsSemantic = bRequirementsPassed;
 	}
 
-	virtual void PostPhysicalDeviceProperties() override final
+	virtual void PostPhysicalDeviceProperties(FOptionalVulkanDeviceExtensions& ExtensionFlags) override final
 	{
 		VkPhysicalDeviceFragmentShaderBarycentricPropertiesKHR& FragmentShaderBarycentricProps = GetDeviceExtensionProperties().FragmentShaderBarycentricProps;
 		//UE_LOG(LogVulkanRHI, Verbose, TEXT("triStripVertexOrderIndependentOfProvokingVertex is %s"), FragmentShaderBarycentricProps.triStripVertexOrderIndependentOfProvokingVertex ? TEXT("true") : TEXT("false"));
@@ -2155,7 +2164,59 @@ public:
 };
 
 
+// ***** VK_EXT_host_image_copy
+class FVulkanEXTHostImageCopyExtension : public FVulkanDeviceExtension
+{
+public:
+
+	FVulkanEXTHostImageCopyExtension(FVulkanDevice* InDevice)
+		: FVulkanDeviceExtension(InDevice, VK_EXT_HOST_IMAGE_COPY_EXTENSION_NAME, VULKAN_EXTENSION_ENABLED, VK_API_VERSION_1_4)
+	{
+		bEnabledInCode = bEnabledInCode && (GVulkanAllowHostImageCopyCVar.GetValueOnAnyThread() != 0);
+	}
+
+	virtual void PrePhysicalDeviceFeatures(VkPhysicalDeviceFeatures2KHR& PhysicalDeviceFeatures2) override final
+	{
+		ZeroVulkanStruct(HostImageCopyFeatures, VK_STRUCTURE_TYPE_PHYSICAL_DEVICE_HOST_IMAGE_COPY_FEATURES);
+		AddToPNext(PhysicalDeviceFeatures2, HostImageCopyFeatures);
+	}
+
+	virtual void PostPhysicalDeviceFeatures(FOptionalVulkanDeviceExtensions& ExtensionFlags) override final
+	{
+		bRequirementsPassed = (HostImageCopyFeatures.hostImageCopy == VK_TRUE);
+	}
+
+	virtual void PrePhysicalDeviceProperties(VkPhysicalDeviceProperties2KHR& PhysicalDeviceProperties2) override final
+	{
+		if (bRequirementsPassed)
+		{
+			ZeroVulkanStruct(HostImageCopyProperties, VK_STRUCTURE_TYPE_PHYSICAL_DEVICE_HOST_IMAGE_COPY_PROPERTIES);
+			AddToPNext(PhysicalDeviceProperties2, HostImageCopyProperties);
+		}
+	}
+
+	virtual void PostPhysicalDeviceProperties(FOptionalVulkanDeviceExtensions& ExtensionFlags) override final
+	{
+		bRequirementsPassed = bRequirementsPassed && (HostImageCopyProperties.identicalMemoryTypeRequirements == VK_TRUE);
+		ExtensionFlags.HasEXTHostImageCopy = bRequirementsPassed;
+		if (bRequirementsPassed)
+		{
+			UE_LOG(LogVulkanRHI, Display, TEXT("Host Image Copy Enabled (SrcCount=%d, DstCount=%d)"), 
+				HostImageCopyProperties.copySrcLayoutCount, HostImageCopyProperties.copyDstLayoutCount);
+		}
+	}
 
+	virtual void PreCreateDevice(VkDeviceCreateInfo& DeviceCreateInfo) override final
+	{
+		if (bRequirementsPassed)
+		{
+			AddToPNext(DeviceCreateInfo, HostImageCopyFeatures);
+		}
+	}
+
+	VkPhysicalDeviceHostImageCopyFeatures HostImageCopyFeatures;
+	VkPhysicalDeviceHostImageCopyProperties HostImageCopyProperties;
+};
 
 
 template <typename ExtensionType>
@@ -2259,6 +2320,7 @@ FVulkanDeviceExtensionArray FVulkanDeviceExtension::GetUESupportedDeviceExtensio
 	ADD_CUSTOM_EXTENSION(FVulkanEXTExtendedDynamicState3Extension);
 	ADD_CUSTOM_EXTENSION(FVulkanEXTVertexInputDynamicStateExtension);
 	ADD_CUSTOM_EXTENSION(FVulkanEXTShaderObjectExtension);
+	ADD_CUSTOM_EXTENSION(FVulkanEXTHostImageCopyExtension);
 
 	// Needed for Raytracing
 	ADD_CUSTOM_EXTENSION(FVulkanKHRBufferDeviceAddressExtension);
diff --git a/Engine/Source/Runtime/VulkanRHI/Private/VulkanExtensions.h b/Engine/Source/Runtime/VulkanRHI/Private/VulkanExtensions.h
--- a/Engine/Source/Runtime/VulkanRHI/Private/VulkanExtensions.h
+++ b/Engine/Source/Runtime/VulkanRHI/Private/VulkanExtensions.h
@@ -113,7 +113,7 @@ public:
 	static TArray<VkExtensionProperties> GetDriverSupportedDeviceExtensions(VkPhysicalDevice Gpu, const ANSICHAR* LayerName = nullptr);
 
 	virtual void PrePhysicalDeviceProperties(VkPhysicalDeviceProperties2KHR& PhysicalDeviceProperties2) {}
-	virtual void PostPhysicalDeviceProperties() {}
+	virtual void PostPhysicalDeviceProperties(FOptionalVulkanDeviceExtensions& ExtensionFlags) {}
 
 	virtual void PrePhysicalDeviceFeatures(VkPhysicalDeviceFeatures2KHR& PhysicalDeviceFeatures2) {}
 	virtual void PostPhysicalDeviceFeatures(FOptionalVulkanDeviceExtensions& ExtensionFlags)
diff --git a/Engine/Source/Runtime/VulkanRHI/Private/VulkanLoader.h b/Engine/Source/Runtime/VulkanRHI/Private/VulkanLoader.h
--- a/Engine/Source/Runtime/VulkanRHI/Private/VulkanLoader.h
+++ b/Engine/Source/Runtime/VulkanRHI/Private/VulkanLoader.h
@@ -270,7 +270,10 @@
 	EnumMacro(PFN_vkCmdSetPrimitiveRestartEnableEXT, vkCmdSetPrimitiveRestartEnableEXT) \
 	EnumMacro(PFN_vkCmdSetAlphaToCoverageEnableEXT, vkCmdSetAlphaToCoverageEnableEXT) \
 	EnumMacro(PFN_vkCmdSetAlphaToOneEnableEXT, vkCmdSetAlphaToOneEnableEXT) \
-	EnumMacro(PFN_vkCmdSetLogicOpEnableEXT, vkCmdSetLogicOpEnableEXT)
+	EnumMacro(PFN_vkCmdSetLogicOpEnableEXT, vkCmdSetLogicOpEnableEXT) \
+	EnumMacro(PFN_vkCopyMemoryToImageEXT, vkCopyMemoryToImageEXT) \
+	EnumMacro(PFN_vkTransitionImageLayoutEXT, vkTransitionImageLayoutEXT)
+
 
 
 
