VulkanRHI: change ray tracing handling in AMDCompatibilityMode

#rnx

diff --git a/Engine/Source/Runtime/VulkanRHI/Private/VulkanRHI.cpp b/Engine/Source/Runtime/VulkanRHI/Private/VulkanRHI.cpp
--- a/Engine/Source/Runtime/VulkanRHI/Private/VulkanRHI.cpp
+++ b/Engine/Source/Runtime/VulkanRHI/Private/VulkanRHI.cpp
@@ -888,27 +888,10 @@ void FVulkanDynamicRHI::InitInstance()
 
 		GRHIGlobals.NeedsExtraTransitions = true;
 
-		// Use this compatibility mode avoid known issues at launch time with latest drivers at the time of release 5.5.  This will:
-		// - disable inline ray tracing and use ray tracing pipelines everywhere (instead of a mix of both)
-		// - disable mesh shaders until issues can be resolved (holes in Nanite meshes)
-		// - disable ray tracing for RADV driver prior to version 24.3.2 in Linux (raytracing pipeline compilation crash)
-
-		// :todo-jn: to be removed when the official minimum RADV version is set to 24.3.2 in BaseHardware.ini
-		const bool bUseAMDCompatibilityMode = GVulkanAMDCompatibilityMode && (Device->GetVendorId() == EGpuVendorId::Amd);
-		if (GVulkanAMDCompatibilityMode && 
-			(Device->GetOptionalExtensionProperties().PhysicalDeviceDriverProperties.driverID == VK_DRIVER_ID_MESA_RADV) &&
-			(Props.driverVersion < VK_MAKE_VERSION(24, 3, 2)))
-		{
-			GRHISupportsRayTracing = false;
-			GRHISupportsRayTracingShaders = false;
-			GRHISupportsInlineRayTracing = false;
-			UE_LOG(LogVulkanRHI, Warning, TEXT("Using MESA RADV version prior to 24.3.2, ray tracing disabled."));
-		}
-
 		if (GRHISupportsRayTracing)
 		{
 			GRHISupportsRayTracingShaders = RHISupportsRayTracingShaders(GMaxRHIShaderPlatform) && Device->GetOptionalExtensions().HasRayTracingPipeline;
-			GRHISupportsInlineRayTracing = !bUseAMDCompatibilityMode && RHISupportsInlineRayTracing(GMaxRHIShaderPlatform) && Device->GetOptionalExtensions().HasRayQuery;
+			GRHISupportsInlineRayTracing = RHISupportsInlineRayTracing(GMaxRHIShaderPlatform) && Device->GetOptionalExtensions().HasRayQuery;
 
 			// Inline RayTracing SBT is needed if raytracing position fetch isn't available
 			GRHIGlobals.RayTracing.RequiresInlineRayTracingSBT = !VULKAN_SUPPORTS_RAY_TRACING_POSITION_FETCH;
@@ -927,6 +910,35 @@ void FVulkanDynamicRHI::InitInstance()
 			GRHIGlobals.RayTracing.SupportsLooseParamsInShaderRecord = true;
 		}
 
+		// Use this compatibility mode avoid known issues at launch time with latest drivers at the time of release 5.5. 
+		// This will disable ray tracing with proprietary driver and older MESA versions
+		const bool bUseAMDCompatibilityMode = GVulkanAMDCompatibilityMode && (Device->GetVendorId() == EGpuVendorId::Amd);
+		if (GVulkanAMDCompatibilityMode)
+		{
+			// :todo-jn: to be removed when the official minimum RADV version is set to 24.3.2 in BaseHardware.ini
+			if (Device->GetOptionalExtensionProperties().PhysicalDeviceDriverProperties.driverID == VK_DRIVER_ID_MESA_RADV)
+			{
+				if (Props.driverVersion < VK_MAKE_VERSION(24, 3, 2))
+				{
+					GRHISupportsRayTracing = false;
+					GRHISupportsRayTracingShaders = false;
+					GRHISupportsInlineRayTracing = false;
+					UE_LOG(LogVulkanRHI, Warning, TEXT("Using MESA RADV version prior to 24.3.2, ray tracing disabled."));
+				}
+			}
+			else if ((Device->GetOptionalExtensionProperties().PhysicalDeviceDriverProperties.driverID == VK_DRIVER_ID_AMD_PROPRIETARY) ||
+					(Device->GetOptionalExtensionProperties().PhysicalDeviceDriverProperties.driverID == VK_DRIVER_ID_AMD_OPEN_SOURCE))
+			{
+				GRHISupportsRayTracing = false;
+				GRHISupportsRayTracingShaders = false;
+				GRHISupportsInlineRayTracing = false;
+				static auto CVarPathTracing = IConsoleManager::Get().FindConsoleVariable(TEXT("r.PathTracing"));
+				CVarPathTracing->Set(0);
+				
+				UE_LOG(LogVulkanRHI, Display, TEXT("Using AMD Propriertary Driver, ray tracing disabled."));
+			}
+		}
+
 #if VULKAN_ENABLE_DUMP_LAYER
 		// Disable RHI thread by default if the dump layer is enabled
 		GRHISupportsRHIThread = false;
