VulkanRHI: Add null check for cvar

#rnx

diff --git a/Engine/Source/Runtime/VulkanRHI/Private/VulkanRHI.cpp b/Engine/Source/Runtime/VulkanRHI/Private/VulkanRHI.cpp
--- a/Engine/Source/Runtime/VulkanRHI/Private/VulkanRHI.cpp
+++ b/Engine/Source/Runtime/VulkanRHI/Private/VulkanRHI.cpp
@@ -933,7 +933,10 @@ void FVulkanDynamicRHI::InitInstance()
 				GRHISupportsRayTracingShaders = false;
 				GRHISupportsInlineRayTracing = false;
 				static auto CVarPathTracing = IConsoleManager::Get().FindConsoleVariable(TEXT("r.PathTracing"));
-				CVarPathTracing->Set(0);
+				if (CVarPathTracing)
+				{
+					CVarPathTracing->Set(0);
+				}
 				
 				UE_LOG(LogVulkanRHI, Display, TEXT("Using AMD Propriertary Driver, ray tracing disabled."));
 			}
