[Linux] Fix FLinearColor readback for Vulkan

[RN] [Linux] Fix loss of precision when FLinearColor is requested in VulkanDynamicRHI::RHIReadSurfaceData

diff --git a/Engine/Source/Runtime/VulkanRHI/Private/VulkanRenderTarget.cpp b/Engine/Source/Runtime/VulkanRHI/Private/VulkanRenderTarget.cpp
--- a/Engine/Source/Runtime/VulkanRHI/Private/VulkanRenderTarget.cpp
+++ b/Engine/Source/Runtime/VulkanRHI/Private/VulkanRenderTarget.cpp
@@ -123,7 +123,54 @@ static void ConvertRawDataToFColor(VkFormat VulkanFormat, uint32 DestWidth, uint
 
 }
 
-void FVulkanDynamicRHI::RHIReadSurfaceData(FRHITexture* TextureRHI, FIntRect Rect, TArray<FColor>& OutData, FReadSurfaceDataFlags InFlags)
+static void ConvertRawDataToFLinearColor(VkFormat VulkanFormat, uint32 DestWidth, uint32 DestHeight, uint8* In, uint32 SrcPitch, FLinearColor* Dest, const FReadSurfaceDataFlags& InFlags)
+{
+	switch (VulkanFormat)
+	{
+	case VK_FORMAT_R32G32B32A32_SFLOAT:
+		ConvertRawR32G32B32A32DataToFLinearColor(DestWidth, DestHeight, In, SrcPitch, Dest, InFlags);
+		break;
+
+	case VK_FORMAT_R16G16B16A16_SFLOAT:
+		ConvertRawR16G16B16A16FDataToFLinearColor(DestWidth, DestHeight, In, SrcPitch, Dest, InFlags);
+		break;
+
+	case VK_FORMAT_B10G11R11_UFLOAT_PACK32:
+		ConvertRawR11G11B10FDataToFLinearColor(DestWidth, DestHeight, In, SrcPitch, Dest);
+		break;
+
+	case VK_FORMAT_A2B10G10R10_UNORM_PACK32:
+		ConvertRawA2B10G10R10DataToFLinearColor(DestWidth, DestHeight, In, SrcPitch, Dest);
+		break;
+
+	case VK_FORMAT_R8G8B8A8_UNORM:
+		ConvertRawR8G8B8A8DataToFLinearColor(DestWidth, DestHeight, In, SrcPitch, Dest);
+		break;
+
+	case VK_FORMAT_R16G16B16A16_UNORM:
+		ConvertRawR16G16B16A16DataToFLinearColor(DestWidth, DestHeight, In, SrcPitch, Dest);
+		break;
+
+	case VK_FORMAT_B8G8R8A8_UNORM:
+		ConvertRawB8G8R8A8DataToFLinearColor(DestWidth, DestHeight, In, SrcPitch, Dest);
+		break;
+
+	case VK_FORMAT_R16G16_UNORM:
+		ConvertRawR16G16DataToFLinearColor(DestWidth, DestHeight, In, SrcPitch, Dest);
+		break;
+
+	case VK_FORMAT_R16_UNORM:
+		ConvertRawR16UDataToFLinearColor(DestWidth, DestHeight, In, SrcPitch, Dest);
+		break;
+
+	default:
+		checkf(false, TEXT("Unsupported format [%d] for conversion to FLinearColor!"), (uint32)VulkanFormat);
+		break;
+	}
+}
+
+template<typename ColorType>
+static void ReadSurfaceData(FVulkanDevice* Device, FRHITexture* TextureRHI, FIntRect Rect, TArray<ColorType>& OutData, FReadSurfaceDataFlags InFlags)
 {
 	checkf((!TextureRHI->GetDesc().IsTextureCube()) || (InFlags.GetCubeFace() == CubeFace_MAX), TEXT("Cube faces not supported yet."));
 
@@ -221,8 +268,15 @@ void FVulkanDynamicRHI::RHIReadSurfaceData(FRHITexture* TextureRHI, FIntRect Rec
 		SrcPitch = DestWidth * PixelByteSize;
 	}
 
-	FColor* Dest = OutData.GetData();
-	ConvertRawDataToFColor(Surface.StorageFormat, DestWidth, DestHeight, In, SrcPitch, Dest, InFlags);
+	ColorType* Dest = OutData.GetData();
+	if constexpr (std::is_same_v<FLinearColor, ColorType>)
+		ConvertRawDataToFLinearColor(Surface.StorageFormat, DestWidth, DestHeight, In, SrcPitch, Dest, InFlags);
+	else if constexpr (std::is_same_v<FColor, ColorType>)
+		ConvertRawDataToFColor(Surface.StorageFormat, DestWidth, DestHeight, In, SrcPitch, Dest, InFlags);
+	else
+	{
+		checkNoEntry();
+	}
 
 	if (!bCPUReadback)
 	{
@@ -230,15 +284,14 @@ void FVulkanDynamicRHI::RHIReadSurfaceData(FRHITexture* TextureRHI, FIntRect Rec
 	}
 }
 
+void FVulkanDynamicRHI::RHIReadSurfaceData(FRHITexture* TextureRHI, FIntRect Rect, TArray<FColor>& OutData, FReadSurfaceDataFlags InFlags)
+{
+	ReadSurfaceData(Device, TextureRHI, Rect, OutData, InFlags);
+}
+
 void FVulkanDynamicRHI::RHIReadSurfaceData(FRHITexture* TextureRHI, FIntRect Rect, TArray<FLinearColor>& OutData, FReadSurfaceDataFlags InFlags)
 {
-	TArray<FColor> FromColorData;
-	RHIReadSurfaceData(TextureRHI, Rect, FromColorData, InFlags);
-	OutData.SetNumUninitialized(FromColorData.Num());
-	for (int Index = 0, Num = FromColorData.Num(); Index < Num; Index++)
-	{
-		OutData[Index] = FLinearColor(FromColorData[Index]);
-	}
+	ReadSurfaceData(Device, TextureRHI, Rect, OutData, InFlags);
 }
 
 void FVulkanDynamicRHI::RHIMapStagingSurface(FRHITexture* TextureRHI, FRHIGPUFence* FenceRHI, void*& OutData, int32& OutWidth, int32& OutHeight, uint32 GPUIndex)
