VulkanRHI: Add support for VK_EXT_host_image_copy on mobile.

#rnx

diff --git a/Engine/Source/Runtime/VulkanRHI/Private/VulkanResources.h b/Engine/Source/Runtime/VulkanRHI/Private/VulkanResources.h
--- a/Engine/Source/Runtime/VulkanRHI/Private/VulkanResources.h
+++ b/Engine/Source/Runtime/VulkanRHI/Private/VulkanResources.h
@@ -623,6 +623,7 @@ public:
 	void AliasTextureResources(FTextureRHIRef& SrcTextureRHI);
 
 	void UploadInitialData(FRHICommandListBase& RHICmdList, VulkanRHI::FStagingBuffer* UploadBuffer);
+	void HostUploadInitialData(FRHICommandListBase& RHICmdList, const uint8* BulkData, uint32 BulkDataSize);
 
 	// View with all mips/layers
 	FVulkanView* DefaultView = nullptr;
diff --git a/Engine/Source/Runtime/VulkanRHI/Private/VulkanTexture.cpp b/Engine/Source/Runtime/VulkanRHI/Private/VulkanTexture.cpp
--- a/Engine/Source/Runtime/VulkanRHI/Private/VulkanTexture.cpp
+++ b/Engine/Source/Runtime/VulkanRHI/Private/VulkanTexture.cpp
@@ -194,6 +194,11 @@ static VkImageUsageFlags GetUsageFlagsFromCreateFlags(FVulkanDevice& InDevice, c
 	{
 		UsageFlags |= VK_IMAGE_USAGE_COLOR_ATTACHMENT_BIT;
 	}
+	else if (InDevice.GetOptionalExtensions().HasEXTHostImageCopy)
+	{
+		// non-targets can be initialized using host transfers
+		UsageFlags |= VK_IMAGE_USAGE_HOST_TRANSFER_BIT;
+	}
 
 	if (EnumHasAnyFlags(UEFlags, TexCreate_Foveation) && ValidateShadingRateDataType())
 	{
@@ -1264,8 +1269,6 @@ namespace UE::VulkanTexture
 
 void FVulkanTexture::UploadInitialData(FRHICommandListBase& RHICmdList, VulkanRHI::FStagingBuffer* UploadBuffer)
 {
-	const FRHITextureDesc& Desc = GetDesc();
-
 	// InternalLockWrite leaves the image in VK_IMAGE_LAYOUT_SHADER_READ_ONLY_OPTIMAL, so make sure the requested resource state is SRV.
 	SetDefaultLayout(VK_IMAGE_LAYOUT_SHADER_READ_ONLY_OPTIMAL);
 
@@ -1332,6 +1335,90 @@ void FVulkanTexture::UploadInitialData(FRHICommandListBase& RHICmdList, VulkanRH
 		});
 }
 
+void FVulkanTexture::HostUploadInitialData(FRHICommandListBase& RHICmdList, const uint8* BulkData, uint32 BulkDataSize)
+{
+	check(Device->GetOptionalExtensions().HasEXTHostImageCopy);
+	check(BulkData);
+
+	// InternalLockWrite leaves the image in VK_IMAGE_LAYOUT_SHADER_READ_ONLY_OPTIMAL, so make sure the requested resource state is SRV.
+	SetDefaultLayout(VK_IMAGE_LAYOUT_SHADER_READ_ONLY_OPTIMAL);
+
+	const FRHITextureDesc& Desc = GetDesc();
+	const uint32 LayerCount = GetNumberOfArrayLevels();
+
+	const VkImageSubresourceRange FullSubresourceRange = {
+			.aspectMask = GetFullAspectMask(),
+			.baseMipLevel = 0,
+			.levelCount = Desc.NumMips,
+			.baseArrayLayer = 0,
+			.layerCount = LayerCount
+	};
+
+	// Transition the entire texture to GENERAL
+	VkHostImageLayoutTransitionInfo TransitionInfo = {
+		.sType = VK_STRUCTURE_TYPE_HOST_IMAGE_LAYOUT_TRANSITION_INFO,
+		.pNext = nullptr,
+		.image = Image,
+		.oldLayout = VK_IMAGE_LAYOUT_UNDEFINED,
+		.newLayout = VK_IMAGE_LAYOUT_GENERAL,
+		.subresourceRange = FullSubresourceRange
+	};
+	VERIFYVULKANRESULT(VulkanRHI::vkTransitionImageLayoutEXT(Device->GetHandle(), 1, &TransitionInfo));
+
+	TArray<VkMemoryToImageCopy, TInlineAllocator<12>> Regions;
+	Regions.AddZeroed(Desc.NumMips);
+	{
+		uint64 MipOffset = 0;
+		for (uint32 MipIndex = 0; MipIndex < Desc.NumMips; MipIndex++)
+		{
+			VkMemoryToImageCopy& Region = Regions[MipIndex];
+
+			check(MipOffset < BulkDataSize);
+
+			const FUintVector3 MipExtents = UE::RHITextureUtils::CalculateMipExtents(Desc, MipIndex);
+
+			Region.pHostPointer = BulkData + MipOffset;
+
+			Region.imageExtent.width = MipExtents.X;
+			Region.imageExtent.height = MipExtents.Y;
+			Region.imageExtent.depth = MipExtents.Z;
+
+			Region.imageSubresource.aspectMask = FullAspectMask;
+			Region.imageSubresource.layerCount = LayerCount;
+			Region.imageSubresource.mipLevel = MipIndex;
+
+			uint64 Stride = 0;
+			const uint64 MipSize = UE::RHITextureUtils::CalculateTextureMipSize(Desc, MipIndex, Stride);
+			const uint64 SubresourceSize = MipSize * LayerCount;
+
+			MipOffset += SubresourceSize;
+		}
+	}
+
+	VkCopyMemoryToImageInfo CopyMemoryToImageInfo = {
+		.sType = VK_STRUCTURE_TYPE_COPY_MEMORY_TO_IMAGE_INFO,
+		.pNext = nullptr,
+		.flags = 0,
+		.dstImage = Image,
+		.dstImageLayout = VK_IMAGE_LAYOUT_GENERAL,
+		.regionCount = (uint32)Regions.Num(),
+		.pRegions = Regions.GetData()
+	};
+	VERIFYVULKANRESULT(VulkanRHI::vkCopyMemoryToImageEXT(Device->GetHandle(), &CopyMemoryToImageInfo));
+
+	// Now queue up a barrier to execute on the GPU to bring it back to a read-only layout
+	{
+		FVulkanUploadContext& Context = FVulkanUploadContext::Get(RHICmdList);
+		FVulkanCommandBuffer& CommandBuffer = Context.GetCommandBuffer();
+		ensure(CommandBuffer.IsOutsideRenderPass());
+
+		FVulkanPipelineBarrier Barrier;
+		Barrier.AddImageLayoutTransition(Image, VK_IMAGE_LAYOUT_GENERAL, VK_IMAGE_LAYOUT_SHADER_READ_ONLY_OPTIMAL, FullSubresourceRange);
+		Barrier.Execute(CommandBuffer.GetHandle());
+	}
+}
+
+
 FVulkanTexture* FVulkanDynamicRHI::CreateTextureInternal(FRHICommandListBase& RHICmdList, const FRHITextureCreateDesc& CreateDesc)
 {
 	FCreateTextureResult CreateResult = BeginCreateTextureInternal(CreateDesc, nullptr);
@@ -1361,55 +1448,110 @@ FRHITextureInitializer FVulkanDynamicRHI::RHICreateTextureInitializer(FRHIComman
 	// TODO: add something to force VK_IMAGE_LAYOUT_SHADER_READ_ONLY_OPTIMAL as the default layout
 
 	FCreateTextureResult CreateResult = BeginCreateTextureInternal(CreateDesc, nullptr);
+	const bool bUseHostImageCopies = Device->GetOptionalExtensions().HasEXTHostImageCopy && 
+		!CreateResult.bClear && CreateResult.Texture && 
+		VKHasAllFlags(CreateResult.Texture->ImageUsageFlags, VK_IMAGE_USAGE_HOST_TRANSFER_BIT);
+	if (bUseHostImageCopies)
+	{
+		// We don't want finalize to transition the image when done from host
+		CreateResult.DefaultLayout = VK_IMAGE_LAYOUT_UNDEFINED;
+	}
 	FVulkanTexture* Texture = FinalizeCreateTextureInternal(RHICmdList, CreateResult);
 
-	const uint64 UploadMemorySize = UE::RHITextureUtils::CalculateTextureSize(CreateDesc);
-	const uint64 VulkanUploadSize = VulkanCalculateTextureSize(Texture);
+	if (bUseHostImageCopies)
+	{
+		if (CreateDesc.InitAction == ERHITextureInitAction::BulkData)
+		{
+			checkf(CreateDesc.BulkData, TEXT("InitAction was BulkData, but none was specified!"));
+			Texture->HostUploadInitialData(RHICmdList, reinterpret_cast<const uint8*>(CreateDesc.BulkData->GetResourceBulkData()), CreateDesc.BulkData->GetResourceBulkDataSize());
 
-	// Transfer bulk data
-	VulkanRHI::FStagingBuffer* UploadBuffer = Device->GetStagingManager().AcquireBuffer(UploadMemorySize);
-	void* const UploadData = UploadBuffer->GetMappedPointer();
+			// Discard the bulk data's contents.
+			CreateDesc.BulkData->Discard();
 
-	UE::RHICore::FBaseTextureInitializerImplementation Initializer(RHICmdList, Texture, UploadData, UploadMemorySize,
-		[Texture = TRefCountPtr<FVulkanTexture>(Texture), UploadBuffer](FRHICommandListBase& RHICmdList) mutable
-		{
-			Texture->UploadInitialData(RHICmdList, UploadBuffer);
-			return TRefCountPtr<FRHITexture>(MoveTemp(Texture));
-		},
-		[Texture = TRefCountPtr<FVulkanTexture>(Texture), UploadBuffer](FRHITextureInitializer::FSubresourceIndex SubresourceIndex) mutable
+			return UE::RHICore::FDefaultTextureInitializer(RHICmdList, Texture);
+		}
+		else if (CreateDesc.InitAction == ERHITextureInitAction::Initializer)
 		{
-			const FRHITextureDesc& TextureDesc = Texture->GetDesc();
+			const uint64 UploadMemorySize = UE::RHITextureUtils::CalculateTextureSize(CreateDesc);
+			uint8* UploadMemory = reinterpret_cast<uint8*>(FMemory::Malloc(UploadMemorySize));
 
-			uint64 SubresourceStride = 0;
-			uint64 SubresourceSize = 0;
-			const uint64 Offset = UE::VulkanTexture::CalculateSubresourceOffset(TextureDesc, SubresourceIndex, SubresourceStride, SubresourceSize);
-			check(Offset + SubresourceSize <= UploadBuffer->GetSize());
+			return UE::RHICore::FBaseTextureInitializerImplementation(RHICmdList, Texture, UploadMemory, UploadMemorySize,
+				[Texture = TRefCountPtr<FVulkanTexture>(Texture), UploadMemory, UploadMemorySize](FRHICommandListBase& RHICmdList) mutable
+				{
+					Texture->HostUploadInitialData(RHICmdList, UploadMemory, UploadMemorySize);
+					FMemory::Free(UploadMemory);
+					return TRefCountPtr<FRHITexture>(MoveTemp(Texture));
+				},
+				[Texture = TRefCountPtr<FVulkanTexture>(Texture), UploadMemory, UploadMemorySize](FRHITextureInitializer::FSubresourceIndex SubresourceIndex) mutable
+				{
+					const FRHITextureDesc& TextureDesc = Texture->GetDesc();
+
+					uint64 SubresourceStride = 0;
+					uint64 SubresourceSize = 0;
+					const uint64 Offset = UE::VulkanTexture::CalculateSubresourceOffset(TextureDesc, SubresourceIndex, SubresourceStride, SubresourceSize);
+					checkf(Offset + SubresourceSize <= UploadMemorySize,
+						TEXT("Subresource(Face=%d,Array=%d,Mip=%d) access out of bounds: Offset=%lld + SubresourceSize=%lld > UploadMemorySize=%lld"), 
+						SubresourceIndex.FaceIndex, SubresourceIndex.ArrayIndex, SubresourceIndex.MipIndex,
+						Offset, SubresourceSize, UploadMemorySize);
+
+					FRHITextureSubresourceInitializer Result{};
+					Result.Data = UploadMemory + Offset;
+					Result.Stride = SubresourceStride;
+					Result.Size = SubresourceSize;
+
+					return Result;
+				}
+			);
+		}
+	}
+	else
+	{
+		const uint64 UploadMemorySize = UE::RHITextureUtils::CalculateTextureSize(CreateDesc);
 
-			FRHITextureSubresourceInitializer Result{};
-			Result.Data = reinterpret_cast<uint8*>(UploadBuffer->GetMappedPointer()) + Offset;
-			Result.Stride = SubresourceStride;
-			Result.Size = SubresourceSize;
+		// Transfer bulk data
+		VulkanRHI::FStagingBuffer* UploadBuffer = Device->GetStagingManager().AcquireBuffer(UploadMemorySize);
+		void* const UploadData = UploadBuffer->GetMappedPointer();
 
-			return Result;
-		}
-	);
+		UE::RHICore::FBaseTextureInitializerImplementation Initializer(RHICmdList, Texture, UploadData, UploadMemorySize,
+			[Texture = TRefCountPtr<FVulkanTexture>(Texture), UploadBuffer](FRHICommandListBase& RHICmdList) mutable
+			{
+				Texture->UploadInitialData(RHICmdList, UploadBuffer);
+				return TRefCountPtr<FRHITexture>(MoveTemp(Texture));
+			},
+			[Texture = TRefCountPtr<FVulkanTexture>(Texture), UploadBuffer](FRHITextureInitializer::FSubresourceIndex SubresourceIndex) mutable
+			{
+				const FRHITextureDesc& TextureDesc = Texture->GetDesc();
 
-	if (CreateDesc.InitAction == ERHITextureInitAction::BulkData)
-	{
-		check(CreateDesc.BulkData);
-		check(UploadMemorySize >= CreateDesc.BulkData->GetResourceBulkDataSize());
+				uint64 SubresourceStride = 0;
+				uint64 SubresourceSize = 0;
+				const uint64 Offset = UE::VulkanTexture::CalculateSubresourceOffset(TextureDesc, SubresourceIndex, SubresourceStride, SubresourceSize);
+				check(Offset + SubresourceSize <= UploadBuffer->GetSize());
 
-		FMemory::Memcpy(UploadData, CreateDesc.BulkData->GetResourceBulkData(), CreateDesc.BulkData->GetResourceBulkDataSize());
+				FRHITextureSubresourceInitializer Result{};
+				Result.Data = reinterpret_cast<uint8*>(UploadBuffer->GetMappedPointer()) + Offset;
+				Result.Stride = SubresourceStride;
+				Result.Size = SubresourceSize;
 
-		// Discard the bulk data's contents.
-		CreateDesc.BulkData->Discard();
+				return Result;
+			}
+		);
 
-		return MoveTemp(Initializer);
-	}
+		if (CreateDesc.InitAction == ERHITextureInitAction::BulkData)
+		{
+			check(CreateDesc.BulkData);
+			check(UploadMemorySize >= CreateDesc.BulkData->GetResourceBulkDataSize());
 
-	if (CreateDesc.InitAction == ERHITextureInitAction::Initializer)
-	{
-		return MoveTemp(Initializer);
+			FMemory::Memcpy(UploadData, CreateDesc.BulkData->GetResourceBulkData(), CreateDesc.BulkData->GetResourceBulkDataSize());
+
+			// Discard the bulk data's contents.
+			CreateDesc.BulkData->Discard();
+
+			return MoveTemp(Initializer);
+		}
+		else if (CreateDesc.InitAction == ERHITextureInitAction::Initializer)
+		{
+			return MoveTemp(Initializer);
+		}
 	}
 
 	return UE::RHICore::HandleUnknownTextureInitializerInitAction(RHICmdList, CreateDesc);
