[Linux] Fix a memzero issue in FVulkanDynamicRHI::RHIReadSurfaceData

diff --git a/Engine/Source/Runtime/VulkanRHI/Private/VulkanRenderTarget.cpp b/Engine/Source/Runtime/VulkanRHI/Private/VulkanRenderTarget.cpp
--- a/Engine/Source/Runtime/VulkanRHI/Private/VulkanRenderTarget.cpp
+++ b/Engine/Source/Runtime/VulkanRHI/Private/VulkanRenderTarget.cpp
@@ -173,6 +173,8 @@ template<typename ColorType>
 static void ReadSurfaceData(FVulkanDevice* Device, FRHITexture* TextureRHI, FIntRect Rect, TArray<ColorType>& OutData, FReadSurfaceDataFlags InFlags)
 {
 	checkf((!TextureRHI->GetDesc().IsTextureCube()) || (InFlags.GetCubeFace() == CubeFace_MAX), TEXT("Cube faces not supported yet."));
+	checkf(((Rect.Max.X > 0) && (Rect.Min.X >=0) && (Rect.Max.X >= Rect.Min.X)), TEXT("Width must be positive."));
+	checkf(((Rect.Max.Y > 0) && (Rect.Min.Y >= 0) && (Rect.Max.Y >= Rect.Min.Y)), TEXT("Height must be positive."));
 
 	const uint32 DestWidth = Rect.Max.X - Rect.Min.X;
 	const uint32 DestHeight = Rect.Max.Y - Rect.Min.Y;
@@ -181,7 +183,7 @@ static void ReadSurfaceData(FVulkanDevice* Device, FRHITexture* TextureRHI, FInt
 	if (GIgnoreCPUReads)
 	{
 		// Debug: Fill with CPU
-		FMemory::Memzero(OutData.GetData(), NumRequestedPixels * sizeof(FColor));
+		FMemory::Memzero(OutData.GetData(), NumRequestedPixels * sizeof(ColorType));
 		return;
 	}
 
@@ -197,7 +199,7 @@ static void ReadSurfaceData(FVulkanDevice* Device, FRHITexture* TextureRHI, FInt
 
 	default:
 		// Just return black for texture types we don't support.
-		FMemory::Memzero(OutData.GetData(), NumRequestedPixels * sizeof(FColor));
+		FMemory::Memzero(OutData.GetData(), NumRequestedPixels * sizeof(ColorType));
 		return;
 	}
 
@@ -270,11 +272,17 @@ static void ReadSurfaceData(FVulkanDevice* Device, FRHITexture* TextureRHI, FInt
 
 	ColorType* Dest = OutData.GetData();
 	if constexpr (std::is_same_v<FLinearColor, ColorType>)
+	{
 		ConvertRawDataToFLinearColor(Surface.StorageFormat, DestWidth, DestHeight, In, SrcPitch, Dest, InFlags);
+	}
 	else if constexpr (std::is_same_v<FColor, ColorType>)
+	{
 		ConvertRawDataToFColor(Surface.StorageFormat, DestWidth, DestHeight, In, SrcPitch, Dest, InFlags);
+	}
 	else
 	{
+		// set output to black for invalid ColorType's
+		FMemory::Memzero(OutData.GetData(), NumRequestedPixels * sizeof(ColorType));
 		checkNoEntry();
 	}
 
