VulkanRHI: Use sentinel when breadcrumb pointer is null.

#rnx

diff --git a/Engine/Source/Runtime/VulkanRHI/Private/VulkanQueue.cpp b/Engine/Source/Runtime/VulkanRHI/Private/VulkanQueue.cpp
--- a/Engine/Source/Runtime/VulkanRHI/Private/VulkanQueue.cpp
+++ b/Engine/Source/Runtime/VulkanRHI/Private/VulkanQueue.cpp
@@ -802,7 +802,7 @@ void FVulkanDiagnosticBuffer::WriteMarkerIn(FVulkanCommandBuffer& CommandBuffer,
 #if NV_AFTERMATH
 	if (Device.GetOptionalExtensions().HasNVDiagnosticCheckpoints)
 	{
-		VulkanDynamicAPI::vkCmdSetCheckpointNV(CommandBufferHandle, Breadcrumb);
+		VulkanDynamicAPI::vkCmdSetCheckpointNV(CommandBufferHandle, Breadcrumb ? Breadcrumb : FRHIBreadcrumbNode::Sentinel);
 	}
 #endif // NV_AFTERMATH
 }
