VulkanRHI: Move null check on breadcrumb earlier

#rnx

diff --git a/Engine/Source/Runtime/VulkanRHI/Private/VulkanQueue.cpp b/Engine/Source/Runtime/VulkanRHI/Private/VulkanQueue.cpp
--- a/Engine/Source/Runtime/VulkanRHI/Private/VulkanQueue.cpp
+++ b/Engine/Source/Runtime/VulkanRHI/Private/VulkanQueue.cpp
@@ -786,6 +786,11 @@ void FVulkanDiagnosticBuffer::WriteMarkerIn(FVulkanCommandBuffer& CommandBuffer,
 	const VkCommandBuffer CommandBufferHandle = CommandBuffer.GetHandle();
 	const uint64 WriteOffset = offsetof(FQueue, MarkerIn);
 
+	if (!Breadcrumb)
+	{
+		Breadcrumb = FRHIBreadcrumbNode::Sentinel;
+	}
+
 	if (Device.GetOptionalExtensions().HasAMDBufferMarker)
 	{
 		VulkanDynamicAPI::vkCmdWriteBufferMarkerAMD(CommandBufferHandle, VK_PIPELINE_STAGE_BOTTOM_OF_PIPE_BIT, Buffer, WriteOffset, Breadcrumb->ID);
@@ -802,7 +807,7 @@ void FVulkanDiagnosticBuffer::WriteMarkerIn(FVulkanCommandBuffer& CommandBuffer,
 #if NV_AFTERMATH
 	if (Device.GetOptionalExtensions().HasNVDiagnosticCheckpoints)
 	{
-		VulkanDynamicAPI::vkCmdSetCheckpointNV(CommandBufferHandle, Breadcrumb ? Breadcrumb : FRHIBreadcrumbNode::Sentinel);
+		VulkanDynamicAPI::vkCmdSetCheckpointNV(CommandBufferHandle, Breadcrumb);
 	}
 #endif // NV_AFTERMATH
 }
@@ -812,6 +817,11 @@ void FVulkanDiagnosticBuffer::WriteMarkerOut(FVulkanCommandBuffer& CommandBuffer
 	const VkCommandBuffer CommandBufferHandle = CommandBuffer.GetHandle();
 	const uint64 WriteOffset = offsetof(FQueue, MarkerOut);
 
+	if (!Breadcrumb)
+	{
+		Breadcrumb = FRHIBreadcrumbNode::Sentinel;
+	}
+
 	if (Device.GetOptionalExtensions().HasAMDBufferMarker)
 	{
 		VulkanDynamicAPI::vkCmdWriteBufferMarkerAMD(CommandBufferHandle, VK_PIPELINE_STAGE_BOTTOM_OF_PIPE_BIT, Buffer, WriteOffset, Breadcrumb->ID);
