VulkanRHI: Fix barrier handle leak for split barriers when transitions don't result in any barriers.  Add stat to track barrier event pool size.

#rnx

diff --git a/Engine/Source/Runtime/VulkanRHI/Private/VulkanBarriers.cpp b/Engine/Source/Runtime/VulkanRHI/Private/VulkanBarriers.cpp
--- a/Engine/Source/Runtime/VulkanRHI/Private/VulkanBarriers.cpp
+++ b/Engine/Source/Runtime/VulkanRHI/Private/VulkanBarriers.cpp
@@ -1445,7 +1445,17 @@ void ProcessTransitionSync2(FVulkanCommandListContext& Context, TArrayView<const
 		if (Data->EventHandle)
 		{
 			checkf(bIsSingleQueue, TEXT("Split barriers must remain on same queue!"));
-			SubmitBarriers(MakeArrayView(BarrierArrays.MemoryBarriers), MakeArrayView(BarrierArrays.BufferBarriers), MakeArrayView(BarrierArrays.ImageBarriers), Data->EventHandle);
+
+			if (BarrierArrays.MemoryBarriers.Num() || BarrierArrays.BufferBarriers.Num() || BarrierArrays.ImageBarriers.Num())
+			{
+				SubmitBarriers(MakeArrayView(BarrierArrays.MemoryBarriers), MakeArrayView(BarrierArrays.BufferBarriers), MakeArrayView(BarrierArrays.ImageBarriers), Data->EventHandle);
+			}
+			else if (!bIsBeginTransition)
+			{
+				// In some cases, transitions result in no barriers (like a destination of Discard for example), make sure we discard the event properly
+				Context.Device.ReleaseBarrierEvent(Data->EventHandle);
+			}
+
 			continue;
 		}
 
diff --git a/Engine/Source/Runtime/VulkanRHI/Private/VulkanDevice.cpp b/Engine/Source/Runtime/VulkanRHI/Private/VulkanDevice.cpp
--- a/Engine/Source/Runtime/VulkanRHI/Private/VulkanDevice.cpp
+++ b/Engine/Source/Runtime/VulkanRHI/Private/VulkanDevice.cpp
@@ -1770,6 +1770,7 @@ VkEvent FVulkanDevice::GetBarrierEvent()
 		ZeroVulkanStruct(Info, VK_STRUCTURE_TYPE_EVENT_CREATE_INFO);
 		Info.flags = VK_EVENT_CREATE_DEVICE_ONLY_BIT;
 		VERIFYVULKANRESULT(VulkanRHI::vkCreateEvent(Device, &Info, VULKAN_CPU_ALLOCATOR, &Handle));
+		INC_DWORD_STAT(STAT_VulkanNumPooledBarrierEvents);
 	}
 
 	return Handle;
diff --git a/Engine/Source/Runtime/VulkanRHI/Private/VulkanRHIPrivate.h b/Engine/Source/Runtime/VulkanRHI/Private/VulkanRHIPrivate.h
--- a/Engine/Source/Runtime/VulkanRHI/Private/VulkanRHIPrivate.h
+++ b/Engine/Source/Runtime/VulkanRHI/Private/VulkanRHIPrivate.h
@@ -202,6 +202,7 @@ DECLARE_CYCLE_STAT_EXTERN(TEXT("VkCreateDescriptorPool"), STAT_VulkanVkCreateDes
 DECLARE_DWORD_ACCUMULATOR_STAT_EXTERN(TEXT("Num Created DescSet Pools"), STAT_VulkanNumDescPools, STATGROUP_VulkanRHI, );
 DECLARE_CYCLE_STAT_EXTERN(TEXT("UpdateUniformBuffers"), STAT_VulkanUpdateUniformBuffers, STATGROUP_VulkanRHI, );
 DECLARE_CYCLE_STAT_EXTERN(TEXT("UpdateUniformBuffersRename"), STAT_VulkanUpdateUniformBuffersRename, STATGROUP_VulkanRHI, );
+DECLARE_DWORD_ACCUMULATOR_STAT_EXTERN(TEXT("Num Pooled Barrier Events"), STAT_VulkanNumPooledBarrierEvents, STATGROUP_VulkanRHI, );
 #if VULKAN_ENABLE_AGGRESSIVE_STATS
 DECLARE_CYCLE_STAT_EXTERN(TEXT("Barrier Time"), STAT_VulkanBarrierTime, STATGROUP_VulkanRHI, );
 DECLARE_CYCLE_STAT_EXTERN(TEXT("Update DescriptorSets"), STAT_VulkanUpdateDescriptorSets, STATGROUP_VulkanRHI, );
diff --git a/Engine/Source/Runtime/VulkanRHI/Private/VulkanUtil.cpp b/Engine/Source/Runtime/VulkanRHI/Private/VulkanUtil.cpp
--- a/Engine/Source/Runtime/VulkanRHI/Private/VulkanUtil.cpp
+++ b/Engine/Source/Runtime/VulkanRHI/Private/VulkanUtil.cpp
@@ -1134,6 +1134,7 @@ DEFINE_STAT(STAT_VulkanVkCreateDescriptorPool);
 DEFINE_STAT(STAT_VulkanNumDescPools);
 DEFINE_STAT(STAT_VulkanUpdateUniformBuffers);
 DEFINE_STAT(STAT_VulkanUpdateUniformBuffersRename);
+DEFINE_STAT(STAT_VulkanNumPooledBarrierEvents);
 #if VULKAN_ENABLE_AGGRESSIVE_STATS
 DEFINE_STAT(STAT_VulkanBarrierTime);
 DEFINE_STAT(STAT_VulkanUpdateDescriptorSets);
