PR #13964: Fix aliased textures to use correct views on vulkan

#rnx

diff --git a/Engine/Source/Runtime/VulkanRHI/Private/VulkanTexture.cpp b/Engine/Source/Runtime/VulkanRHI/Private/VulkanTexture.cpp
--- a/Engine/Source/Runtime/VulkanRHI/Private/VulkanTexture.cpp
+++ b/Engine/Source/Runtime/VulkanRHI/Private/VulkanTexture.cpp
@@ -2412,6 +2412,11 @@ void FVulkanTexture::UpdateViews(const FVulkanContextArray& Contexts, VkImageUsa
 static FRWLock GInternalViewRWLock; // Use single global lock for now
 FVulkanView* FVulkanTexture::FindOrAddInternalView(const FVulkanTextureViewDesc& ViewDesc)
 {
+	if ((ImageOwnerType == EImageOwnerType::Aliased) && AliasedTexture.IsValid())
+	{
+		return ResourceCast(AliasedTexture)->FindOrAddInternalView(ViewDesc);
+	}
+
 	{
 		FRWScopeLock ScopedLock(GInternalViewRWLock, SLT_ReadOnly);
 		if (FVulkanView** VulkanView = InternalViews.Find(ViewDesc.GetHash()))
