To ensure the cooker never picks up a preview shader and preview modes never pickup standard shaders,
always set a value for the preferred preview shader code format in all RHIs that support preview modes.

#test Switched between multiple preview platforms (PC D3D SM5, PC D3D Mobile, Generic iOS, Generic Android) on DX11 and DX12 RHIs. Repeated several times with editor closing. Also checked the new substring is included in the DDC keys and ran a local cook.
#rnx

diff --git a/Engine/Source/Runtime/Apple/MetalRHI/Private/MetalRHI.cpp b/Engine/Source/Runtime/Apple/MetalRHI/Private/MetalRHI.cpp
--- a/Engine/Source/Runtime/Apple/MetalRHI/Private/MetalRHI.cpp
+++ b/Engine/Source/Runtime/Apple/MetalRHI/Private/MetalRHI.cpp
@@ -499,6 +499,9 @@ FMetalDynamicRHI::FMetalDynamicRHI(ERHIFeatureLevel::Type RequestedFeatureLevel)
 		GMaxRHIShaderPlatform = SP_METAL_SM5;
 	}
 	
+	// Metal always uses AIR. This is used to differentiate standard and preview shaders in the DDC.
+	GRHIGlobals.PreferredPreviewShaderCodeFormat = TEXT("AIR");
+
 #if PLATFORM_SUPPORTS_MESH_SHADERS
     GRHISupportsMeshShadersTier0 = RHISupportsMeshShadersTier0(GMaxRHIShaderPlatform);
     GRHISupportsMeshShadersTier1 = RHISupportsMeshShadersTier1(GMaxRHIShaderPlatform);
diff --git a/Engine/Source/Runtime/D3D12RHI/Private/D3D12Adapter.cpp b/Engine/Source/Runtime/D3D12RHI/Private/D3D12Adapter.cpp
--- a/Engine/Source/Runtime/D3D12RHI/Private/D3D12Adapter.cpp
+++ b/Engine/Source/Runtime/D3D12RHI/Private/D3D12Adapter.cpp
@@ -1295,6 +1295,12 @@ void FD3D12Adapter::InitializeDevices()
 				{
 					GRHIGlobals.PreferredPreviewShaderCodeFormat = TEXT("DXIL");
 				}
+				else
+				{
+					// If SM 6.0 is not supported, we cannot use DXIL and therefore have to use the same code format as D3D11, which is DXBC.
+					// In this case, preview shaders can potentially be shared between D3D11 and D3D12 just like they can be shared in a cook for standard (i.e. non-preview) shaders.
+					GRHIGlobals.PreferredPreviewShaderCodeFormat = TEXT("DXBC");
+				}
 
 				if (D3D12Caps7.MeshShaderTier >= D3D12_MESH_SHADER_TIER_1)
 				{
diff --git a/Engine/Source/Runtime/OpenGLDrv/Private/OpenGLDevice.cpp b/Engine/Source/Runtime/OpenGLDrv/Private/OpenGLDevice.cpp
--- a/Engine/Source/Runtime/OpenGLDrv/Private/OpenGLDevice.cpp
+++ b/Engine/Source/Runtime/OpenGLDrv/Private/OpenGLDevice.cpp
@@ -1146,6 +1146,9 @@ static void InitRHICapabilitiesForGL()
 	GRHISupportsPipelineFileCache = !GRHISupportsPSOPrecaching || CVarEnablePSOFileCacheWhenPrecachingActive.GetValueOnAnyThread();
 
 	GRHIGlobals.NeedsShaderUnbinds = true;
+
+	// OpenGL always uses GLSL (aka GLSL ES or ESSL). This is used to differentiate standard and preview shaders in the DDC.
+	GRHIGlobals.PreferredPreviewShaderCodeFormat = TEXT("GLSL");
 }
 
 FDynamicRHI* FOpenGLDynamicRHIModule::CreateRHI(ERHIFeatureLevel::Type InRequestedFeatureLevel)
diff --git a/Engine/Source/Runtime/VulkanRHI/Private/VulkanRHI.cpp b/Engine/Source/Runtime/VulkanRHI/Private/VulkanRHI.cpp
--- a/Engine/Source/Runtime/VulkanRHI/Private/VulkanRHI.cpp
+++ b/Engine/Source/Runtime/VulkanRHI/Private/VulkanRHI.cpp
@@ -450,6 +450,9 @@ FVulkanDynamicRHI::FVulkanDynamicRHI()
 	GRHISupportsPipelineFileCache = !GRHISupportsPSOPrecaching || CVarEnableVulkanPSOFileCacheWhenPrecachingActive.GetValueOnAnyThread();
 	UE_LOG(LogVulkanRHI, Log, TEXT("Vulkan PSO Precaching = %d, PipelineFileCache = %d"), GRHISupportsPSOPrecaching, GRHISupportsPipelineFileCache);
 
+	// Vulkan always uses SPIR-V. This is used to differentiate standard and preview shaders in the DDC.
+	GRHIGlobals.PreferredPreviewShaderCodeFormat = TEXT("SPIRV");
+
 	// Copy source requires its own image layout.
 	EnumRemoveFlags(GRHIMergeableAccessMask, ERHIAccess::CopySrc);
 
diff --git a/Engine/Source/Runtime/Windows/D3D11RHI/Private/D3D11Device.cpp b/Engine/Source/Runtime/Windows/D3D11RHI/Private/D3D11Device.cpp
--- a/Engine/Source/Runtime/Windows/D3D11RHI/Private/D3D11Device.cpp
+++ b/Engine/Source/Runtime/Windows/D3D11RHI/Private/D3D11Device.cpp
@@ -219,6 +219,9 @@ FD3D11DynamicRHI::FD3D11DynamicRHI(IDXGIFactory1* InDXGIFactory1, D3D_FEATURE_LE
 	GRHIGlobals.MaxViewSizeBytesForNonTypedBuffer = 1 << 27;
 	GRHIGlobals.MaxViewDimensionForTypedBuffer = 1 << 27;
 	
+	// D3D11 always uses DXBC. This is used to differentiate standard and preview shaders in the DDC.
+	GRHIGlobals.PreferredPreviewShaderCodeFormat = TEXT("DXBC");
+
 	// Initialize the constant buffers.
 	InitConstantBuffers();
 
