VulkanRHI: Integrate 46676867 (use ref instead of pointer) and 46672738 (fix signed/unsigned mismatch) from Main.  Fix crash for certain race conditions (made obvious by async compute).

#rnx

diff --git a/Engine/Source/Runtime/VulkanRHI/Private/VulkanMemory.h b/Engine/Source/Runtime/VulkanRHI/Private/VulkanMemory.h
--- a/Engine/Source/Runtime/VulkanRHI/Private/VulkanMemory.h
+++ b/Engine/Source/Runtime/VulkanRHI/Private/VulkanMemory.h
@@ -859,7 +859,7 @@ namespace VulkanRHI
 		struct FUBPendingFree
 		{
 			FVulkanAllocation Allocation;
-			uint64 Frame = 0;
+			int64 Frame = 0;
 		};
 
 		struct
diff --git a/Engine/Source/Runtime/VulkanRHI/Private/VulkanSubmission.cpp b/Engine/Source/Runtime/VulkanRHI/Private/VulkanSubmission.cpp
--- a/Engine/Source/Runtime/VulkanRHI/Private/VulkanSubmission.cpp
+++ b/Engine/Source/Runtime/VulkanRHI/Private/VulkanSubmission.cpp
@@ -614,7 +614,7 @@ void FVulkanDynamicRHI::EnqueueEndOfPipeTask(TUniqueFunction<void()> TaskFunc, T
 
 void FVulkanDynamicRHI::RHIBlockUntilGPUIdle()
 {
-	TArray<FVulkanSyncPoint*, TInlineAllocator<(int32)EVulkanQueueType::Count>> EndOfQueueEvent;
+	TArray<FVulkanSyncPointRef, TInlineAllocator<(int32)EVulkanQueueType::Count>> EndOfQueueEvent;
 
 	// Create payloads with a signal for each queue
 	{
@@ -637,7 +637,7 @@ void FVulkanDynamicRHI::RHIBlockUntilGPUIdle()
 	KickSubmissionThread();
 
 	// Wait on each event
-	for (FVulkanSyncPoint* Event : EndOfQueueEvent)
+	for (FVulkanSyncPointRef& Event : EndOfQueueEvent)
 	{
 		ProcessInterruptQueueUntil(Event);
 	}
