Adds console variable r.Vulkan.RequiresWaitingForFrameCompletionEvent to optionally require waiting for frame completion

diff --git a/Engine/Source/Runtime/VulkanRHI/Private/Android/VulkanAndroidPlatform.cpp b/Engine/Source/Runtime/VulkanRHI/Private/Android/VulkanAndroidPlatform.cpp
--- a/Engine/Source/Runtime/VulkanRHI/Private/Android/VulkanAndroidPlatform.cpp
+++ b/Engine/Source/Runtime/VulkanRHI/Private/Android/VulkanAndroidPlatform.cpp
@@ -88,6 +88,15 @@ static TAutoConsoleVariable<int32> CVarVulkanHasUnifiedMemory(
 );
 #endif
 
+FAutoConsoleVariableRef FVulkanAndroidPlatform::CVarVulkanRequiresWaitingForFrameCompletionEvent(
+	TEXT("r.Vulkan.RequiresWaitingForFrameCompletionEvent"),
+	FVulkanAndroidPlatform::bRequiresWaitingForFrameCompletionEvent,
+	TEXT("Whether Android VulkanRHI should wait for frame completion event.")
+	TEXT("  0 = don't wait (default)\n")
+	TEXT("  1 = wait"),
+	ECVF_RenderThreadSafe
+);
+
 // Vulkan function pointers
 #define DEFINE_VK_ENTRYPOINTS(Type,Func) Type VulkanDynamicAPI::Func = NULL;
 ENUM_VK_ENTRYPOINTS_ALL(DEFINE_VK_ENTRYPOINTS)
@@ -119,6 +128,7 @@ int32 FVulkanAndroidPlatform::AFBCWorkaroundOption = 0;
 int32 FVulkanAndroidPlatform::ASTCWorkaroundOption = 0;
 bool FVulkanAndroidPlatform::bRequiresDepthStencilFullWrite = false;
 bool FVulkanAndroidPlatform::bRequiresMergedRenderPassDepthStencilLayout = false;
+bool FVulkanAndroidPlatform::bRequiresWaitingForFrameCompletionEvent = false;
 
 #define CHECK_VK_ENTRYPOINTS(Type,Func) if (VulkanDynamicAPI::Func == NULL) { bFoundAllEntryPoints = false; UE_LOG(LogRHI, Warning, TEXT("Failed to find entry point for %s"), TEXT(#Func)); }
 
diff --git a/Engine/Source/Runtime/VulkanRHI/Private/Android/VulkanAndroidPlatform.h b/Engine/Source/Runtime/VulkanRHI/Private/Android/VulkanAndroidPlatform.h
--- a/Engine/Source/Runtime/VulkanRHI/Private/Android/VulkanAndroidPlatform.h
+++ b/Engine/Source/Runtime/VulkanRHI/Private/Android/VulkanAndroidPlatform.h
@@ -200,7 +200,7 @@ public:
 	//#todo-rco: Detect Mali? Doing a clear on ColorAtt layout on empty cmd buffer causes issues
 	static bool RequiresSwapchainGeneralInitialLayout() { return true; }
 
-	static bool RequiresWaitingForFrameCompletionEvent() { return false; }
+	static bool RequiresWaitingForFrameCompletionEvent() { return bRequiresWaitingForFrameCompletionEvent; }
 	
 	// Does the platform allow a nullptr Pixelshader on the pipeline
 	static bool SupportsNullPixelShader() { return false; }
@@ -290,6 +290,9 @@ protected:
 
 	static bool bRequiresDepthStencilFullWrite;
 	static bool bRequiresMergedRenderPassDepthStencilLayout;
+	static bool bRequiresWaitingForFrameCompletionEvent;
+
+	static FAutoConsoleVariableRef CVarVulkanRequiresWaitingForFrameCompletionEvent;
 };
 
 #if VULKAN_SUPPORTS_GOOGLE_DISPLAY_TIMING
