PR #14138: VulkanRHI Fix R16_SFLOAT color conversion to FLinearColor

#rnx

diff --git a/Engine/Source/Runtime/VulkanRHI/Private/VulkanRenderTarget.cpp b/Engine/Source/Runtime/VulkanRHI/Private/VulkanRenderTarget.cpp
--- a/Engine/Source/Runtime/VulkanRHI/Private/VulkanRenderTarget.cpp
+++ b/Engine/Source/Runtime/VulkanRHI/Private/VulkanRenderTarget.cpp
@@ -135,6 +135,10 @@ static void ConvertRawDataToFLinearColor(VkFormat VulkanFormat, uint32 DestWidth
 		ConvertRawR16G16B16A16FDataToFLinearColor(DestWidth, DestHeight, In, SrcPitch, Dest, InFlags);
 		break;
 
+	case VK_FORMAT_R16_SFLOAT:
+		ConvertRawR16FDataToFLinearColor(DestWidth, DestHeight, In, SrcPitch, Dest);
+		break;
+
 	case VK_FORMAT_B10G11R11_UFLOAT_PACK32:
 		ConvertRawR11G11B10FDataToFLinearColor(DestWidth, DestHeight, In, SrcPitch, Dest);
 		break;
