VulkanRHI: Delay swapchain creation when NumFormats==0 and add more logging. Fix leaked surface on some swapchain creation errors.  Use all images in the swapchain (skipping images had improper synchronization).  Use proper TRefCountPtr on semaphores instead of manually adding/releasing refs.

#rnx

#Synced-CL 50685202

diff --git a/Engine/Source/Runtime/VulkanRHI/Private/VulkanSwapChain.cpp b/Engine/Source/Runtime/VulkanRHI/Private/VulkanSwapChain.cpp
--- a/Engine/Source/Runtime/VulkanRHI/Private/VulkanSwapChain.cpp
+++ b/Engine/Source/Runtime/VulkanRHI/Private/VulkanSwapChain.cpp
@@ -164,8 +164,16 @@ FVulkanSwapChain* FVulkanSwapChain::Create(
 #define VERIFYVULKANSURFACERESULT(VkFunction) \
 	{\
 		const VkResult ScopedResult = VkFunction; \
-		if (ScopedResult == VK_ERROR_SURFACE_LOST_KHR) { return nullptr; } \
-		if (ScopedResult < VK_SUCCESS) { VulkanRHI::VerifyVulkanResult(ScopedResult, #VkFunction, __FILE__, __LINE__); } \
+		if (ScopedResult == VK_ERROR_SURFACE_LOST_KHR) \
+		{ \
+			DestroyRecreateInfo(); \
+			VulkanRHI::vkDestroySurfaceKHR(InInstance, Surface, VULKAN_CPU_ALLOCATOR); \
+			return nullptr; \
+		} \
+		else if (ScopedResult < VK_SUCCESS) \
+		{ \
+			VulkanRHI::VerifyVulkanResult(ScopedResult, #VkFunction, __FILE__, __LINE__); \
+		} \
 	}
 
 	VkSurfaceKHR Surface = VK_NULL_HANDLE;
@@ -182,7 +190,7 @@ FVulkanSwapChain* FVulkanSwapChain::Create(
 		FVulkanPlatform::CreateSurface(WindowContext, InInstance, &Surface);
 	}
 
-	ON_SCOPE_EXIT
+	auto DestroyRecreateInfo = [RecreateInfo, InInstance, &Device]()
 	{
 		if (RecreateInfo != nullptr)
 		{
@@ -199,6 +207,11 @@ FVulkanSwapChain* FVulkanSwapChain::Create(
 		}
 	};
 
+	ON_SCOPE_EXIT
+	{
+		DestroyRecreateInfo();
+	};
+
 #if !UE_BUILD_SHIPPING
 	if (GSimulateLostSurfaceInNextSwapChain)
 	{
@@ -209,7 +222,15 @@ FVulkanSwapChain* FVulkanSwapChain::Create(
 
 	uint32 NumFormats;
 	VERIFYVULKANSURFACERESULT(VulkanRHI::vkGetPhysicalDeviceSurfaceFormatsKHR(Device.GetPhysicalHandle(), Surface, &NumFormats, nullptr));
-	check(NumFormats > 0);
+	if (NumFormats == 0)
+	{
+		UE_LOG(LogVulkanRHI, Warning, TEXT("vkGetPhysicalDeviceSurfaceFormatsKHR returned 0 formats! Handling as surface lost."));
+		VERIFYVULKANSURFACERESULT(VK_ERROR_SURFACE_LOST_KHR);
+	}
+	else if (InOutDesiredNumBackBuffers && (NumFormats < *InOutDesiredNumBackBuffers))
+	{
+		UE_LOG(LogVulkanRHI, Warning, TEXT("vkGetPhysicalDeviceSurfaceFormatsKHR returned %u formats, expected %u!"), NumFormats, *InOutDesiredNumBackBuffers);
+	}
 
 	TArray<VkSurfaceFormatKHR> Formats;
 	Formats.AddZeroed(NumFormats);
@@ -631,14 +652,13 @@ FVulkanSwapChain::FVulkanSwapChain(
 	FVulkanFenceManager& FenceMgr = Device.GetFenceManager();
 	for (uint32 BufferIndex = 0; BufferIndex < NumSwapChainImages; ++BufferIndex)
 	{
-		ImageAcquiredFences[BufferIndex] = Device.GetFenceManager().AllocateFence(true);
+		ImageAcquiredFences[BufferIndex] = FenceMgr.AllocateFence(true);
 	}
 #endif
-	ImageAcquiredSemaphore.AddUninitialized(NumSwapChainImages);
+	ImageAcquiredSemaphore.SetNum(NumSwapChainImages);
 	for (uint32 BufferIndex = 0; BufferIndex < NumSwapChainImages; ++BufferIndex)
 	{
 		ImageAcquiredSemaphore[BufferIndex] = new FVulkanSemaphore(Device, EVulkanSemaphoreFlags::ExternallySignaled);
-		ImageAcquiredSemaphore[BufferIndex]->AddRef();
 	}
 }
 
@@ -648,7 +668,7 @@ void FVulkanSwapChain::Destroy(FVulkanSwapChainRecreateInfo* RecreateInfo)
 	// Alternatively could also check on the fence(s) for the image(s) from the swapchain but then timing out/waiting could become an issue.
 	Device.WaitUntilIdle();
 
-	bool bRecreate = RecreateInfo && GVulkanKeepSwapChain;
+	const bool bRecreate = RecreateInfo && GVulkanKeepSwapChain;
 	if (bRecreate)
 	{
 		RecreateInfo->SwapChain = SwapChain;
@@ -666,13 +686,10 @@ void FVulkanSwapChain::Destroy(FVulkanSwapChainRecreateInfo* RecreateInfo)
 	{
 		FenceMgr.ReleaseFence(ImageAcquiredFences[Index]);
 	}
+	ImageAcquiredFences.Empty();
 #endif
 
-	//#todo-rco: Enqueue for deletion as we first need to destroy the cmd buffers and queues otherwise validation fails
-	for (int BufferIndex = 0; BufferIndex < ImageAcquiredSemaphore.Num(); ++BufferIndex)
-	{
-		ImageAcquiredSemaphore[BufferIndex]->Release();
-	}
+	ImageAcquiredSemaphore.Empty();
 
 	if (!bRecreate)
 	{
@@ -704,8 +721,6 @@ int32 FVulkanSwapChain::AcquireImageIndex(FVulkanSemaphore** OutSemaphore)
 #endif
 	VkResult Result;
 	{
-		const uint32 MaxImageIndex = ImageAcquiredSemaphore.Num() - 1;
-
 		SCOPE_CYCLE_COUNTER(STAT_VulkanAcquireBackBuffer);
 		FRenderThreadIdleScope IdleScope(ERenderThreadIdleTypes::WaitingForGPUPresent);
 
@@ -716,18 +731,6 @@ int32 FVulkanSwapChain::AcquireImageIndex(FVulkanSemaphore** OutSemaphore)
 			ImageAcquiredSemaphore[SemaphoreIndex]->GetHandle(),
 			AcquiredFence,
 			&ImageIndex);
-
-		// The swapchain may have more images than we have requested on creating it. Ignore all extra images
-		while (ImageIndex > MaxImageIndex && (Result == VK_SUCCESS || Result == VK_SUBOPTIMAL_KHR))
-		{
-			Result = VulkanRHI::vkAcquireNextImageKHR(
-				Device.GetHandle(),
-				SwapChain,
-				UINT64_MAX,
-				ImageAcquiredSemaphore[SemaphoreIndex]->GetHandle(),
-				AcquiredFence,
-				&ImageIndex);
-		}
 	}
 
 	if (Result == VK_ERROR_OUT_OF_DATE_KHR)
@@ -764,7 +767,7 @@ int32 FVulkanSwapChain::AcquireImageIndex(FVulkanSemaphore** OutSemaphore)
 #if VULKAN_USE_IMAGE_ACQUIRE_FENCES
 	{
 		SCOPE_CYCLE_COUNTER(STAT_VulkanWaitSwapchain);
-		bool bResult = FenceMgr.WaitForFence(ImageAcquiredFences[SemaphoreIndex], UINT64_MAX);
+		const bool bResult = FenceMgr.WaitForFence(ImageAcquiredFences[SemaphoreIndex], UINT64_MAX);
 		ensure(bResult);
 	}
 #endif
diff --git a/Engine/Source/Runtime/VulkanRHI/Private/VulkanSwapChain.h b/Engine/Source/Runtime/VulkanRHI/Private/VulkanSwapChain.h
--- a/Engine/Source/Runtime/VulkanRHI/Private/VulkanSwapChain.h
+++ b/Engine/Source/Runtime/VulkanRHI/Private/VulkanSwapChain.h
@@ -81,7 +81,7 @@ protected:
 	
 	double NextPresentTargetTime = 0;
 
-	TArray<FVulkanSemaphore*> ImageAcquiredSemaphore;
+	TArray<TRefCountPtr<FVulkanSemaphore>> ImageAcquiredSemaphore;
 #if VULKAN_USE_IMAGE_ACQUIRE_FENCES
 	TArray<FVulkanFence*> ImageAcquiredFences;
 #endif
diff --git a/Engine/Source/Runtime/VulkanRHI/Private/VulkanViewport.cpp b/Engine/Source/Runtime/VulkanRHI/Private/VulkanViewport.cpp
--- a/Engine/Source/Runtime/VulkanRHI/Private/VulkanViewport.cpp
+++ b/Engine/Source/Runtime/VulkanRHI/Private/VulkanViewport.cpp
@@ -156,15 +156,17 @@ FVulkanViewport::~FVulkanViewport()
 	if (SupportsStandardSwapchain())
 	{
 		TextureViews.Empty();
+		RenderingDoneSemaphores.Empty();
 
-		for (int32 Index = 0, NumBuffers = RenderingDoneSemaphores.Num(); Index < NumBuffers; ++Index)
+		for (int32 Index = 0; Index < BackBufferImages.Num(); ++Index)
 		{
-			RenderingDoneSemaphores[Index]->Release();
-
 			// FIXME: race condition on TransitionAndLayoutManager, could this be called from RT while RHIT is active?
-			Device.NotifyDeletedImage(BackBufferImages[Index]->Image, true);
-			BackBufferImages[Index] = nullptr;
+			if (BackBufferImages[Index].IsValid())
+			{
+				Device.NotifyDeletedImage(BackBufferImages[Index]->Image, true);
+			}
 		}
+		BackBufferImages.Empty();
 
 		if (SwapChain)
 		{
@@ -379,6 +381,18 @@ void FVulkanViewport::InitImages(FVulkanContextCommon& Context, TConstArrayView<
 	VkClearColorValue ClearColor;
 	FMemory::Memzero(ClearColor);
 
+	const FVulkanTextureViewDesc ViewDesc =
+	{
+		.ViewType = VK_IMAGE_VIEW_TYPE_2D,
+		.AspectFlags = VK_IMAGE_ASPECT_COLOR_BIT,
+		.UEFormat = PixelFormat,
+		.Format = UEToVkTextureFormat(PixelFormat, false),
+		.FirstMip = 0,
+		.NumMips = 1,
+		.ArraySliceIndex = 0,
+		.NumArraySlices = 1
+	};
+
 	const VkImageSubresourceRange Range = FVulkanPipelineBarrier::MakeSubresourceRange(VK_IMAGE_ASPECT_COLOR_BIT);
 	for (int32 Index = 0; Index < Images.Num(); ++Index)
 	{
@@ -391,16 +405,6 @@ void FVulkanViewport::InitImages(FVulkanContextCommon& Context, TConstArrayView<
 		}
 		BackBufferImages[Index] = (FVulkanTexture*)FVulkanDynamicRHI::Get().RHICreateTexture2DFromResource(PixelFormat, ImageSizeX, ImageSizeY, 1, 1, Images[Index], ETextureCreateFlags::RenderTargetable | ETextureCreateFlags::Presentable).GetReference();
 
-		FVulkanTextureViewDesc ViewDesc;
-		ViewDesc.ViewType = VK_IMAGE_VIEW_TYPE_2D;
-		ViewDesc.AspectFlags = VK_IMAGE_ASPECT_COLOR_BIT;
-		ViewDesc.UEFormat = PixelFormat;
-		ViewDesc.Format = UEToVkTextureFormat(PixelFormat, false);
-		ViewDesc.FirstMip = 0;
-		ViewDesc.NumMips = 1;
-		ViewDesc.ArraySliceIndex = 0;
-		ViewDesc.NumArraySlices = 1;
-
 		FVulkanView* NewView = new FVulkanView(Device, VK_DESCRIPTOR_TYPE_SAMPLED_IMAGE);
 		NewView->CreateTextureView(Images[Index], ViewDesc);
 
@@ -434,7 +438,7 @@ void FVulkanViewport::CreateSwapchain(FVulkanCommandListContext& Context, FVulka
 
 		if (WindowContext.IsValid())
 		{
-			uint32 DesiredNumBackBuffers = NUM_BUFFERS;
+			uint32 DesiredNumBackBuffers = NumRequestedSwapchainImages;
 			TArray<VkImage> Images;
 			SwapChain = FVulkanSwapChain::Create(
 				FVulkanDynamicRHI::Get().Instance, Device,
@@ -448,25 +452,21 @@ void FVulkanViewport::CreateSwapchain(FVulkanCommandListContext& Context, FVulka
 
 			if (SwapChain)
 			{
-				checkf(Images.Num() >= NUM_BUFFERS, TEXT("We wanted at least %i images, actual Num: %i"), NUM_BUFFERS, Images.Num());
-
-				const bool bCreateSemaphores = RenderingDoneSemaphores.IsEmpty();
-				checkf(bCreateSemaphores || RenderingDoneSemaphores.Num() == Images.Num(),
-					TEXT("CreateSwapchain, image count is not expected to change:  Requested=%d  Actual=%d  Previous=%d"),
-					NUM_BUFFERS, Images.Num(), RenderingDoneSemaphores.Num());
+				ensureMsgf(Images.Num() >= NumRequestedSwapchainImages,
+					TEXT("SwapChain Images: Requested Count = %i images , actual count: %i , previous count: %i"), 
+					NumRequestedSwapchainImages, Images.Num(), RenderingDoneSemaphores.Num());
 
 				BackBufferImages.SetNum(Images.Num());
 				RenderingDoneSemaphores.SetNum(Images.Num());
-				InitImages(Context, Images);
-
-				if (bCreateSemaphores)
+				for (int32 Index = 0; Index < RenderingDoneSemaphores.Num(); ++Index)
 				{
-					for (int32 Index = 0, NumBuffers = RenderingDoneSemaphores.Num(); Index < NumBuffers; ++Index)
+					if (!RenderingDoneSemaphores[Index].IsValid())
 					{
 						RenderingDoneSemaphores[Index] = new FVulkanSemaphore(Device);
-						RenderingDoneSemaphores[Index]->AddRef();
 					}
 				}
+
+				InitImages(Context, Images);
 			}
 		}
 
@@ -539,10 +539,13 @@ void FVulkanViewport::DestroySwapchain(FVulkanSwapChainRecreateInfo* RecreateInf
 	if (SupportsStandardSwapchain() && SwapChain)
 	{
 		TextureViews.Empty();
-		for (int32 Index = 0, NumBuffers = BackBufferImages.Num(); Index < NumBuffers; ++Index)
+		for (int32 Index = 0; Index < BackBufferImages.Num(); ++Index)
 		{
-			Device.NotifyDeletedImage(BackBufferImages[Index]->Image, true);
-			BackBufferImages[Index] = nullptr;
+			if (BackBufferImages[Index].IsValid())
+			{
+				Device.NotifyDeletedImage(BackBufferImages[Index]->Image, true);
+				BackBufferImages[Index] = nullptr;
+			}
 		}
 		
 		Device.GetDeferredDeletionQueue().ReleaseResources(true);
diff --git a/Engine/Source/Runtime/VulkanRHI/Private/VulkanViewport.h b/Engine/Source/Runtime/VulkanRHI/Private/VulkanViewport.h
--- a/Engine/Source/Runtime/VulkanRHI/Private/VulkanViewport.h
+++ b/Engine/Source/Runtime/VulkanRHI/Private/VulkanViewport.h
@@ -43,8 +43,6 @@ private:
 class FVulkanViewport : public FRHIViewport
 {
 public:
-	enum { NUM_BUFFERS = 3 };
-
 	FVulkanViewport(FVulkanDevice& InDevice, void* InWindowHandle, uint32 InSizeX,uint32 InSizeY,bool bInIsFullscreen, EPixelFormat InPreferredPixelFormat);
 	~FVulkanViewport();
 
@@ -111,13 +109,15 @@ public:
 	void DestroySwapchain(struct FVulkanSwapChainRecreateInfo* RecreateInfo);
 	void RecreateSwapchain(FVulkanCommandListContext& Context, FVulkanPlatformWindowContext& WindowContext);
 
+	// This is the number of images we're aiming for, driver might return more
+	constexpr static uint32 NumRequestedSwapchainImages = 3;
+
 protected:
 	FVulkanDevice& Device;
 
-	// NUM_BUFFERS don't have to match exactly as the driver can require a minimum number larger than NUM_BUFFERS. Provide some slack
-	TArray<TRefCountPtr<FVulkanTexture>, TInlineAllocator<NUM_BUFFERS*2>> BackBufferImages;
-	TArray<FVulkanSemaphore*, TInlineAllocator<NUM_BUFFERS*2>> RenderingDoneSemaphores;
-	TIndirectArray<FVulkanView, TInlineAllocator<NUM_BUFFERS*2>> TextureViews;
+	TArray<TRefCountPtr<FVulkanTexture>, TInlineAllocator<NumRequestedSwapchainImages*2>> BackBufferImages;
+	TArray<TRefCountPtr<FVulkanSemaphore>, TInlineAllocator<NumRequestedSwapchainImages*2>> RenderingDoneSemaphores;
+	TIndirectArray<FVulkanView, TInlineAllocator<NumRequestedSwapchainImages*2>> TextureViews;
 	TRefCountPtr<FVulkanBackBuffer> RHIBackBuffer;
 
 	// 'Dummy' back buffer
