VulkanRHI: Fix log spam

diff --git a/Engine/Source/Runtime/VulkanRHI/Private/VulkanSwapChain.cpp b/Engine/Source/Runtime/VulkanRHI/Private/VulkanSwapChain.cpp
--- a/Engine/Source/Runtime/VulkanRHI/Private/VulkanSwapChain.cpp
+++ b/Engine/Source/Runtime/VulkanRHI/Private/VulkanSwapChain.cpp
@@ -229,7 +229,7 @@ FVulkanSwapChain* FVulkanSwapChain::Create(
 	}
 	else if (InOutDesiredNumBackBuffers && (NumFormats < *InOutDesiredNumBackBuffers))
 	{
-		UE_LOG(LogVulkanRHI, Warning, TEXT("vkGetPhysicalDeviceSurfaceFormatsKHR returned %u formats, expected %u!"), NumFormats, *InOutDesiredNumBackBuffers);
+		// UE_LOG(LogVulkanRHI, Warning, TEXT("vkGetPhysicalDeviceSurfaceFormatsKHR returned %u formats, expected %u!"), NumFormats, *InOutDesiredNumBackBuffers);
 	}
 
 	TArray<VkSurfaceFormatKHR> Formats;
