Ensure PSOs are always created when the vulkan chunked cache is in use but the pipeline LRU is disabled.
#tests running local client

diff --git a/Engine/Source/Runtime/VulkanRHI/Private/VulkanPipeline.cpp b/Engine/Source/Runtime/VulkanRHI/Private/VulkanPipeline.cpp
--- a/Engine/Source/Runtime/VulkanRHI/Private/VulkanPipeline.cpp
+++ b/Engine/Source/Runtime/VulkanRHI/Private/VulkanPipeline.cpp
@@ -3041,12 +3041,8 @@ void FVulkanPipelineStateCacheManager::LRUAdd(FVulkanGraphicsPipelineState* PSO)
 
 void FVulkanPipelineStateCacheManager::LRUTouch(FVulkanGraphicsPipelineState* PSO)
 {
-	if(!bUseLRU)
-	{
-		return;
-	}
 	FScopeLock Lock(&LRUCS);
-	check((PSO->GetVulkanPipeline() == 0) == (PSO->LRUNode == 0));
+	check(!bUseLRU || ((PSO->GetVulkanPipeline() == 0) == (PSO->LRUNode == 0)));
 
 	if(PSO->LRUNode)
 	{
@@ -3090,7 +3086,7 @@ void FVulkanPipelineStateCacheManager::LRUTouch(FVulkanGraphicsPipelineState* PS
 		}
 		else
 		{
-			check(PSO->LRUNode);
+			check(!bUseLRU || PSO->LRUNode);
 		}
 	}
 }
