VulkanRHI: keep descriptor set layouts used to create pipeline layouts, the are needed for shader objects.

#rnx

diff --git a/Engine/Source/Runtime/VulkanRHI/Private/VulkanBindlessDescriptorManager.cpp b/Engine/Source/Runtime/VulkanRHI/Private/VulkanBindlessDescriptorManager.cpp
--- a/Engine/Source/Runtime/VulkanRHI/Private/VulkanBindlessDescriptorManager.cpp
+++ b/Engine/Source/Runtime/VulkanRHI/Private/VulkanBindlessDescriptorManager.cpp
@@ -623,18 +623,18 @@ void FVulkanBindlessDescriptorManager::Init()
 
 	// Now create the single pipeline layout used by everything
 	{
-		VkDescriptorSetLayout DescriptorSetLayouts[VulkanBindless::MaxNumSets];
 		for (int32 LayoutIndex = 0; LayoutIndex < VulkanBindless::NumBindlessSets; ++LayoutIndex)
 		{
 			const BindlessSetState& State = BindlessSetStates[LayoutIndex];
-			DescriptorSetLayouts[LayoutIndex] = State.DescriptorSetLayout;
+			BindlessDescriptorSetLayouts[LayoutIndex] = State.DescriptorSetLayout;
 		}
-		DescriptorSetLayouts[VulkanBindless::BindlessSingleUseUniformBufferSet] = SingleUseUBDescriptorSetLayout;
+		BindlessDescriptorSetLayouts[VulkanBindless::BindlessSingleUseUniformBufferSet] = SingleUseUBDescriptorSetLayout;
+		check(VulkanBindless::BindlessSingleUseUniformBufferSet < BindlessDescriptorSetLayouts.Num());
 
 		VkPipelineLayoutCreateInfo PipelineLayoutCreateInfo;
 		ZeroVulkanStruct(PipelineLayoutCreateInfo, VK_STRUCTURE_TYPE_PIPELINE_LAYOUT_CREATE_INFO);
-		PipelineLayoutCreateInfo.setLayoutCount = VulkanBindless::MaxNumSets;
-		PipelineLayoutCreateInfo.pSetLayouts = DescriptorSetLayouts;
+		PipelineLayoutCreateInfo.setLayoutCount = BindlessDescriptorSetLayouts.Num();
+		PipelineLayoutCreateInfo.pSetLayouts = BindlessDescriptorSetLayouts.GetData();
 
 		// Set the static uniform buffers with push constants (only enabled for ray tracing in DDPI, only used for RGS)
 		VkPushConstantRange PushConstants;
diff --git a/Engine/Source/Runtime/VulkanRHI/Private/VulkanBindlessDescriptorManager.h b/Engine/Source/Runtime/VulkanRHI/Private/VulkanBindlessDescriptorManager.h
--- a/Engine/Source/Runtime/VulkanRHI/Private/VulkanBindlessDescriptorManager.h
+++ b/Engine/Source/Runtime/VulkanRHI/Private/VulkanBindlessDescriptorManager.h
@@ -34,6 +34,12 @@ public:
 		return BindlessPipelineLayout;
 	}
 
+	// Returns descriptor set layouts for a single stage
+	TConstArrayView<VkDescriptorSetLayout> GetDescriptorSetLayouts() const
+	{
+		return BindlessDescriptorSetLayouts;
+	}
+
 	void BindDescriptorBuffers(VkCommandBuffer CommandBuffer, VkPipelineStageFlags SupportedStages);
 
 	FRHIDescriptorHandle AllocateDescriptor(VkDescriptorType DescriptorType);
@@ -97,6 +103,7 @@ private:
 	uint32_t BufferIndices[VulkanBindless::MaxNumSets];
 
 	VkPipelineLayout BindlessPipelineLayout = VK_NULL_HANDLE;
+	TStaticArray<VkDescriptorSetLayout, VulkanBindless::MaxNumSets> BindlessDescriptorSetLayouts;
 
 	BindlessSetState& GetBindlessState(ERHIDescriptorType DescriptorType);
 	BindlessSetState& GetBindlessState(VkDescriptorType DescriptorType);
