VulkanRHI: Add more dynamic state functions that are used with shader objects.

#rnx

diff --git a/Engine/Source/Runtime/VulkanRHI/Private/VulkanCommandWrappers.h b/Engine/Source/Runtime/VulkanRHI/Private/VulkanCommandWrappers.h
--- a/Engine/Source/Runtime/VulkanRHI/Private/VulkanCommandWrappers.h
+++ b/Engine/Source/Runtime/VulkanRHI/Private/VulkanCommandWrappers.h
@@ -295,6 +295,8 @@ struct FWrapLayer
 	static void CmdSetLogicOpEnableEXT(VkResult Result, VkCommandBuffer CommandBuffer, VkBool32 LogicOpEnable) VULKAN_LAYER_BODY
 	static void CmdSetAlphaToCoverageEnableEXT(VkResult Result, VkCommandBuffer CommandBuffer, VkBool32 AlphaToCoverageEnable) VULKAN_LAYER_BODY
 	static void CmdSetAlphaToOneEnableEXT(VkResult Result, VkCommandBuffer CommandBuffer, VkBool32 AlphaToOneEnable) VULKAN_LAYER_BODY
+	static void CmdSetSampleMaskEXT(VkResult Result, VkCommandBuffer CommandBuffer, VkSampleCountFlagBits Samples, const VkSampleMask* SampleMask) VULKAN_LAYER_BODY
+	static void CmdSetFragmentShadingRateKHR(VkResult Result, VkCommandBuffer CommandBuffer, const VkExtent2D* FragmentSize, const VkFragmentShadingRateCombinerOpKHR combinerOps[2]) VULKAN_LAYER_BODY
 
 	static void CopyMemoryToImageEXT(VkResult Result, VkDevice Device, const VkCopyMemoryToImageInfo* CopyMemoryToImageInfo) VULKAN_LAYER_BODY
 	static void TransitionImageLayoutEXT(VkResult Result, VkDevice Device, uint32_t TransitionCount, const VkHostImageLayoutTransitionInfo* Transitions) VULKAN_LAYER_BODY
@@ -2080,6 +2082,20 @@ namespace VulkanRHI
 		FWrapLayer::CmdSetAlphaToOneEnableEXT(VK_SUCCESS, CommandBuffer, AlphaToOneEnable);
 	}
 
+	static FORCEINLINE_DEBUGGABLE void vkCmdSetSampleMaskEXT(VkCommandBuffer CommandBuffer, VkSampleCountFlagBits Samples, const VkSampleMask* SampleMask)
+	{
+		FWrapLayer::CmdSetSampleMaskEXT(VK_RESULT_MAX_ENUM, CommandBuffer, Samples, SampleMask);
+		VULKANAPINAMESPACE::vkCmdSetSampleMaskEXT(CommandBuffer, Samples, SampleMask);
+		FWrapLayer::CmdSetSampleMaskEXT(VK_SUCCESS, CommandBuffer, Samples, SampleMask);
+	}
+
+	static FORCEINLINE_DEBUGGABLE void vkCmdSetFragmentShadingRateKHR(VkCommandBuffer CommandBuffer, const VkExtent2D* FragmentSize, const VkFragmentShadingRateCombinerOpKHR combinerOps[2])
+	{
+		FWrapLayer::CmdSetFragmentShadingRateKHR(VK_RESULT_MAX_ENUM, CommandBuffer, FragmentSize, combinerOps);
+		VULKANAPINAMESPACE::vkCmdSetFragmentShadingRateKHR(CommandBuffer, FragmentSize, combinerOps);
+		FWrapLayer::CmdSetFragmentShadingRateKHR(VK_SUCCESS, CommandBuffer, FragmentSize, combinerOps);
+	}
+
 	static FORCEINLINE_DEBUGGABLE VkResult vkCopyMemoryToImageEXT(VkDevice Device, const VkCopyMemoryToImageInfo* CopyMemoryToImageInfo)
 	{
 		FWrapLayer::CopyMemoryToImageEXT(VK_RESULT_MAX_ENUM, Device, CopyMemoryToImageInfo);
diff --git a/Engine/Source/Runtime/VulkanRHI/Private/VulkanDebug.cpp b/Engine/Source/Runtime/VulkanRHI/Private/VulkanDebug.cpp
--- a/Engine/Source/Runtime/VulkanRHI/Private/VulkanDebug.cpp
+++ b/Engine/Source/Runtime/VulkanRHI/Private/VulkanDebug.cpp
@@ -4995,6 +4995,26 @@ void FWrapLayer::CmdSetAlphaToOneEnableEXT(VkResult Result, VkCommandBuffer Comm
 	}
 }
 
+void FWrapLayer::CmdSetSampleMaskEXT(VkResult Result, VkCommandBuffer CommandBuffer, VkSampleCountFlagBits Samples, const VkSampleMask* SampleMask)
+{
+	if (Result == VK_RESULT_MAX_ENUM)
+	{
+#if VULKAN_ENABLE_DUMP_LAYER
+		PrintfBeginResult(FString::Printf(TEXT("CmdSetSampleMaskEXT()")));
+#endif
+	}
+}
+
+void FWrapLayer::CmdSetFragmentShadingRateKHR(VkResult Result, VkCommandBuffer CommandBuffer, const VkExtent2D* FragmentSize, const VkFragmentShadingRateCombinerOpKHR combinerOps[2])
+{
+	if (Result == VK_RESULT_MAX_ENUM)
+	{
+#if VULKAN_ENABLE_DUMP_LAYER
+		PrintfBeginResult(FString::Printf(TEXT("CmdSetFragmentShadingRateKHR()")));
+#endif
+	}
+}
+
 void FWrapLayer::CopyMemoryToImageEXT(VkResult Result, VkDevice Device, const VkCopyMemoryToImageInfo* CopyMemoryToImageInfo)
 {
 	if (Result == VK_RESULT_MAX_ENUM)
diff --git a/Engine/Source/Runtime/VulkanRHI/Private/VulkanLoader.h b/Engine/Source/Runtime/VulkanRHI/Private/VulkanLoader.h
--- a/Engine/Source/Runtime/VulkanRHI/Private/VulkanLoader.h
+++ b/Engine/Source/Runtime/VulkanRHI/Private/VulkanLoader.h
@@ -271,6 +271,8 @@
 	EnumMacro(PFN_vkCmdSetAlphaToCoverageEnableEXT, vkCmdSetAlphaToCoverageEnableEXT) \
 	EnumMacro(PFN_vkCmdSetAlphaToOneEnableEXT, vkCmdSetAlphaToOneEnableEXT) \
 	EnumMacro(PFN_vkCmdSetLogicOpEnableEXT, vkCmdSetLogicOpEnableEXT) \
+	EnumMacro(PFN_vkCmdSetSampleMaskEXT, vkCmdSetSampleMaskEXT) \
+	EnumMacro(PFN_vkCmdSetFragmentShadingRateKHR, vkCmdSetFragmentShadingRateKHR) \
 	EnumMacro(PFN_vkCopyMemoryToImageEXT, vkCopyMemoryToImageEXT) \
 	EnumMacro(PFN_vkTransitionImageLayoutEXT, vkTransitionImageLayoutEXT)
 
