VulkanRHI: Set unused dynamic states when shader objects are enabled.  Add shader object handle in FVulkanShader.  Add support for adding lambdas in the FDeferredDeletionQueue.

#rnx

diff --git a/Engine/Source/Runtime/VulkanRHI/Private/VulkanCommandBuffer.cpp b/Engine/Source/Runtime/VulkanRHI/Private/VulkanCommandBuffer.cpp
--- a/Engine/Source/Runtime/VulkanRHI/Private/VulkanCommandBuffer.cpp
+++ b/Engine/Source/Runtime/VulkanRHI/Private/VulkanCommandBuffer.cpp
@@ -232,9 +232,36 @@ void FVulkanCommandBuffer::Begin(FVulkanQueryPool* OptionalQueryPool, VkRenderPa
 	}
 #endif
 
+	SetUnusedDynamicStates();
+
 	bNeedsFullDynamicStateUpdate = true;
 }
 
+void FVulkanCommandBuffer::SetUnusedDynamicStates()
+{
+	// Some unused states have been made dynamic with shader objects, set a default value at the start of each new command buffer
+	// :todo-jn: Look for a way to batch these settings (it was suggested to use a baked secondary command buffer)
+	if (Device.SupportsShaderObjects())
+	{
+		VulkanRHI::vkCmdSetRasterizerDiscardEnableEXT(CommandBufferHandle, false);
+		VulkanRHI::vkCmdSetDepthClampEnableEXT(CommandBufferHandle, false);
+		VulkanRHI::vkCmdSetAlphaToOneEnableEXT(CommandBufferHandle, false);
+		VulkanRHI::vkCmdSetLogicOpEnableEXT(CommandBufferHandle, false);
+		VulkanRHI::vkCmdSetPrimitiveRestartEnableEXT(CommandBufferHandle, false);
+		VulkanRHI::vkCmdSetFrontFaceEXT(CommandBufferHandle, VK_FRONT_FACE_COUNTER_CLOCKWISE);
+		VulkanRHI::vkCmdSetLineWidth(CommandBufferHandle, 1.0f);
+		VulkanRHI::vkCmdSetStencilCompareMask(CommandBufferHandle, VK_STENCIL_FACE_FRONT_AND_BACK, MAX_uint32);
+		VulkanRHI::vkCmdSetStencilWriteMask(CommandBufferHandle, VK_STENCIL_FACE_FRONT_AND_BACK, MAX_uint32);
+
+		const VkSampleMask SampleMask = MAX_uint32;
+		VulkanRHI::vkCmdSetSampleMaskEXT(CommandBufferHandle, VK_SAMPLE_COUNT_1_BIT, &SampleMask);
+
+		// Clear unused stages
+		const VkShaderStageFlagBits ShaderStages[] = { VK_SHADER_STAGE_TESSELLATION_CONTROL_BIT, VK_SHADER_STAGE_TESSELLATION_EVALUATION_BIT };
+		VulkanRHI::vkCmdBindShadersEXT(CommandBufferHandle, UE_ARRAY_COUNT(ShaderStages), ShaderStages, nullptr);
+	}
+}
+
 void FVulkanCommandBuffer::Reset()
 {
 	// Reset the secondary command buffers we executed from this one
diff --git a/Engine/Source/Runtime/VulkanRHI/Private/VulkanCommandBuffer.h b/Engine/Source/Runtime/VulkanRHI/Private/VulkanCommandBuffer.h
--- a/Engine/Source/Runtime/VulkanRHI/Private/VulkanCommandBuffer.h
+++ b/Engine/Source/Runtime/VulkanRHI/Private/VulkanCommandBuffer.h
@@ -176,6 +176,8 @@ public:
 	FVulkanDevice& Device;
 
 private:
+	void SetUnusedDynamicStates();
+
 	VkCommandBuffer CommandBufferHandle = VK_NULL_HANDLE;
 	double SubmittedTime = 0.0f;
 
diff --git a/Engine/Source/Runtime/VulkanRHI/Private/VulkanDevice.h b/Engine/Source/Runtime/VulkanRHI/Private/VulkanDevice.h
--- a/Engine/Source/Runtime/VulkanRHI/Private/VulkanDevice.h
+++ b/Engine/Source/Runtime/VulkanRHI/Private/VulkanDevice.h
@@ -226,6 +226,7 @@ namespace VulkanRHI
 			BufferSuballocation,
 			AccelerationStructure,
 			BindlessHandle,
+			Func
 		};
 
 		template <typename T>
@@ -244,6 +245,7 @@ namespace VulkanRHI
 			}
 		}
 
+		void EnqueueFunc(TUniqueFunction<void(class FVulkanDevice&)>&& Func);
 		void EnqueueResourceAllocation(FVulkanAllocation& Allocation);
 		void EnqueueDeviceAllocation(FDeviceMemoryAllocation* DeviceMemoryAllocation);
 
@@ -261,9 +263,10 @@ namespace VulkanRHI
 		{
 			EType StructureType;
 			uint32 FrameNumber;
-			uint64 Handle;
+			uint64 Handle = 0;
 			FVulkanAllocation Allocation;
-			FDeviceMemoryAllocation* DeviceMemoryAllocation;
+			FDeviceMemoryAllocation* DeviceMemoryAllocation = nullptr;
+			TUniqueFunction<void(class FVulkanDevice&)>* Func = nullptr;
 		};
 
 		void ReleaseResourcesImmediately(const TArray<FEntry>& InEntries);
@@ -571,7 +574,22 @@ public:
 			OptionalDeviceExtensions.HasKHRDynamicRendering &&  
 			// in bindless a single pipeline layout is needed
 			SupportsBindless();
+	}
 
+	bool SupportsShaderObjects() const
+	{
+		return OptionalDeviceExtensions.HasEXTShaderObject &&
+			// Extended dynamic states make it possible to create the pipeline library at shader creation,
+			// and they allow for a single library for vertex_input and fragment_output_state
+			OptionalDeviceExtensions.HasEXTExtendedDynamicState1 &&
+			OptionalDeviceExtensions.HasEXTExtendedDynamicState2 &&
+			OptionalDeviceExtensions.HasEXTExtendedDynamicState3 &&
+			// use a single vertex input stage for all pipelines
+			OptionalDeviceExtensions.HasEXTVertexInputDynamicState &&
+			// dynamic rendering removes need for renderpass at pipeline creation
+			OptionalDeviceExtensions.HasKHRDynamicRendering &&  
+			// in bindless a single pipeline layout is needed
+			SupportsBindless();
 	}
 
 	void SetupPresentQueue(VkSurfaceKHR Surface);
diff --git a/Engine/Source/Runtime/VulkanRHI/Private/VulkanMemory.cpp b/Engine/Source/Runtime/VulkanRHI/Private/VulkanMemory.cpp
--- a/Engine/Source/Runtime/VulkanRHI/Private/VulkanMemory.cpp
+++ b/Engine/Source/Runtime/VulkanRHI/Private/VulkanMemory.cpp
@@ -4539,7 +4539,6 @@ namespace VulkanRHI
 		FEntry Entry;
 		Entry.StructureType = Type;
 		Entry.FrameNumber = GVulkanRHIDeletionFrameNumber;
-		Entry.DeviceMemoryAllocation = 0;
 		Entry.Handle = Handle;
 		{
 			FScopeLock ScopeLock(&CS);
@@ -4556,6 +4555,18 @@ namespace VulkanRHI
 		}
 	}
 
+	void FDeferredDeletionQueue2::EnqueueFunc(TUniqueFunction<void(FVulkanDevice&)>&& Func)
+	{
+		FEntry Entry;
+		Entry.StructureType = EType::Func;
+		Entry.FrameNumber = GVulkanRHIDeletionFrameNumber;
+		Entry.Func = new TUniqueFunction<void(FVulkanDevice&)>(MoveTemp(Func));
+		{
+			FScopeLock ScopeLock(&CS);
+			Entries.Add(Entry);
+		}
+	}
+
 	void FDeferredDeletionQueue2::EnqueueDeviceAllocation(FDeviceMemoryAllocation* DeviceMemoryAllocation)
 	{
 		check(DeviceMemoryAllocation);
@@ -4563,8 +4574,6 @@ namespace VulkanRHI
 		FEntry Entry;
 		Entry.StructureType = EType::DeviceMemoryAllocation;
 		Entry.FrameNumber = GVulkanRHIDeletionFrameNumber;
-
-		Entry.Handle = 0;
 		Entry.DeviceMemoryAllocation = DeviceMemoryAllocation;
 		{
 			FScopeLock ScopeLock(&CS);
@@ -4583,14 +4592,9 @@ namespace VulkanRHI
 		FEntry Entry;
 		Entry.StructureType = EType::ResourceAllocation;
 		Entry.FrameNumber = GVulkanRHIDeletionFrameNumber;
-
-		Entry.Handle = 0;
 		Entry.Allocation = Allocation;
-		Entry.DeviceMemoryAllocation = 0;
-
 		{
 			FScopeLock ScopeLock(&CS);
-
 			Entries.Add(Entry);
 		}
 		check(!Allocation.HasAllocation());
@@ -4693,6 +4697,12 @@ namespace VulkanRHI
 				Device.GetBindlessDescriptorManager()->FreeDescriptor(DescriptorHandle);
 				break;
 			}
+			case EType::Func:
+			{
+				(*Entry.Func)(Device);
+				delete Entry.Func;
+				break;
+			}
 
 			default:
 				check(0);
diff --git a/Engine/Source/Runtime/VulkanRHI/Private/VulkanPendingState.cpp b/Engine/Source/Runtime/VulkanRHI/Private/VulkanPendingState.cpp
--- a/Engine/Source/Runtime/VulkanRHI/Private/VulkanPendingState.cpp
+++ b/Engine/Source/Runtime/VulkanRHI/Private/VulkanPendingState.cpp
@@ -782,6 +782,38 @@ void FVulkanPendingGfxState::UpdateDynamicStates(FVulkanCommandBuffer& CommandBu
 		}
 	}
 
+	if (Device.SupportsShaderObjects())
+	{
+		if (GRHISupportsPipelineVariableRateShading && (GRHIVariableRateShadingImageDataType == VRSImage_Palette))
+		{
+			VULKAN_DYNAMIC_STATE_BLOCK_START(ShadingRate, RenderTargets.bHasFragmentDensityAttachment, Combiner);
+			{
+				const bool bNewUsingVariableRateShading = (NewDynamicState.ShadingRate != EVRSShadingRate::VRSSR_1x1) ||
+					(NewDynamicState.RenderTargets.bHasFragmentDensityAttachment && NewDynamicState.Combiner != EVRSRateCombiner::VRSRB_Passthrough);
+
+				if (bNewUsingVariableRateShading)
+				{
+					extern VkFragmentShadingRateCombinerOpKHR GetFragmentCombinerOp(EVRSRateCombiner RateCombiner);
+
+					const VkExtent2D FragmentSize = Device.GetBestMatchedFragmentSize(NewDynamicState.ShadingRate);
+					VkFragmentShadingRateCombinerOpKHR CombinerOps[2] =
+					{
+						GetFragmentCombinerOp(NewDynamicState.Combiner),
+						VK_FRAGMENT_SHADING_RATE_COMBINER_OP_MAX_KHR // @todo: This needs to be specified too.
+					};
+					VulkanRHI::vkCmdSetFragmentShadingRateKHR(CommandBuffer.GetHandle(), &FragmentSize, CombinerOps);
+				}
+				else
+				{
+					const VkExtent2D FragmentSize = { 1, 1 };
+					const VkFragmentShadingRateCombinerOpKHR CombinerOps[] = { VK_FRAGMENT_SHADING_RATE_COMBINER_OP_KEEP_KHR , VK_FRAGMENT_SHADING_RATE_COMBINER_OP_KEEP_KHR };
+					VulkanRHI::vkCmdSetFragmentShadingRateKHR(CommandBuffer.GetHandle(), &FragmentSize, CombinerOps);
+				}
+			}
+			VULKAN_DYNAMIC_STATE_BLOCK_END();
+		}
+	}
+
 #undef VULKAN_DYNAMIC_STATE_COMPARE
 #undef VULKAN_DYNAMIC_STATE_UPDATE
 #undef VULKAN_DYNAMIC_STATE_BLOCK_START
diff --git a/Engine/Source/Runtime/VulkanRHI/Private/VulkanPipeline.cpp b/Engine/Source/Runtime/VulkanRHI/Private/VulkanPipeline.cpp
--- a/Engine/Source/Runtime/VulkanRHI/Private/VulkanPipeline.cpp
+++ b/Engine/Source/Runtime/VulkanRHI/Private/VulkanPipeline.cpp
@@ -1226,6 +1226,11 @@ static const TMap<uint8, VkFragmentShadingRateCombinerOpKHR> FragmentCombinerOpM
 	// @todo: Add "VK_FRAGMENT_SHADING_RATE_COMBINER_OP_MUL_KHR"?
 };
 
+VkFragmentShadingRateCombinerOpKHR GetFragmentCombinerOp(EVRSRateCombiner RateCombiner)
+{
+	return FragmentCombinerOpMap[(uint8)RateCombiner];
+}
+
 static FString GfxShaderHashesToString(FVulkanShader* Shaders[ShaderStage::NumGraphicsStages])
 {
 	FString ShaderHashes = "";
diff --git a/Engine/Source/Runtime/VulkanRHI/Private/VulkanResources.h b/Engine/Source/Runtime/VulkanRHI/Private/VulkanResources.h
--- a/Engine/Source/Runtime/VulkanRHI/Private/VulkanResources.h
+++ b/Engine/Source/Runtime/VulkanRHI/Private/VulkanResources.h
@@ -181,6 +181,7 @@ public:
 		TArray<uint32> UncompressedCode;
 	public:
 		TArrayView<uint32> GetCodeView() {return CodeView;}
+		int32 GetCodeSizeInBytes() {return CodeView.Num() * sizeof(uint32);}
 	};
 
 	inline FSpirvCode GetSpirvCode() const
@@ -216,6 +217,7 @@ protected:
 	bool							bUsesBindless = false;
 
 	VkPipeline						PipelineLibrary = VK_NULL_HANDLE;
+	VkShaderEXT						ShaderObject = VK_NULL_HANDLE;
 
 protected:
 	class FSpirvContainer
@@ -249,6 +251,7 @@ protected:
 	friend class FVulkanComputeShaderState;
 	friend class FVulkanComputePipeline;
 	friend class FVulkanShaderFactory;
+	friend class FVulkanShaderObjectManager;
 };
 
 /** This represents a vertex shader that hasn't been combined with a specific declaration to create a bound shader. */
