VulkanRHI: Remove the need to explicitely enable dynamic states in order to use shader objects.

#rnx

diff --git a/Engine/Source/Runtime/VulkanRHI/Private/VulkanCommandBuffer.cpp b/Engine/Source/Runtime/VulkanRHI/Private/VulkanCommandBuffer.cpp
--- a/Engine/Source/Runtime/VulkanRHI/Private/VulkanCommandBuffer.cpp
+++ b/Engine/Source/Runtime/VulkanRHI/Private/VulkanCommandBuffer.cpp
@@ -28,7 +28,7 @@ FVulkanCommandBuffer::FVulkanCommandBuffer(FVulkanDevice& InDevice, FVulkanComma
 	FScopeLock ScopeLock(CommandBufferPool.GetCS());
 	AllocMemory();
 
-	if (InDevice.GetOptionalExtensions().HasAnyExtendedDynamicState())
+	if (InDevice.GetOptionalExtensions().HasAnyExtendedDynamicState() || InDevice.SupportsShaderObjects())
 	{
 		LastDynamicStateUpdate = new FGfxPipelineDesc();
 		FMemory::Memzero(*LastDynamicStateUpdate);
@@ -247,8 +247,9 @@ void FVulkanCommandBuffer::SetUnusedDynamicStates()
 		VulkanRHI::vkCmdSetDepthClampEnableEXT(CommandBufferHandle, false);
 		VulkanRHI::vkCmdSetAlphaToOneEnableEXT(CommandBufferHandle, false);
 		VulkanRHI::vkCmdSetLogicOpEnableEXT(CommandBufferHandle, false);
+		//VulkanRHI::vkCmdSetLogicOpEXT(CommandBufferHandle, VK_LOGIC_OP_CLEAR);  // ignored when LogicOpEnable is false
 		VulkanRHI::vkCmdSetPrimitiveRestartEnableEXT(CommandBufferHandle, false);
-		VulkanRHI::vkCmdSetFrontFaceEXT(CommandBufferHandle, VK_FRONT_FACE_COUNTER_CLOCKWISE);
+		VulkanRHI::vkCmdSetFrontFaceEXT(CommandBufferHandle, VK_FRONT_FACE_CLOCKWISE);
 		VulkanRHI::vkCmdSetLineWidth(CommandBufferHandle, 1.0f);
 		VulkanRHI::vkCmdSetStencilCompareMask(CommandBufferHandle, VK_STENCIL_FACE_FRONT_AND_BACK, MAX_uint32);
 		VulkanRHI::vkCmdSetStencilWriteMask(CommandBufferHandle, VK_STENCIL_FACE_FRONT_AND_BACK, MAX_uint32);
diff --git a/Engine/Source/Runtime/VulkanRHI/Private/VulkanDevice.h b/Engine/Source/Runtime/VulkanRHI/Private/VulkanDevice.h
--- a/Engine/Source/Runtime/VulkanRHI/Private/VulkanDevice.h
+++ b/Engine/Source/Runtime/VulkanRHI/Private/VulkanDevice.h
@@ -578,14 +578,8 @@ public:
 
 	bool SupportsShaderObjects() const
 	{
+		// NOTE: No need to check individual dynamic state extension, they all get enabled by turning on shader objects
 		return OptionalDeviceExtensions.HasEXTShaderObject &&
-			// Extended dynamic states make it possible to create the pipeline library at shader creation,
-			// and they allow for a single library for vertex_input and fragment_output_state
-			OptionalDeviceExtensions.HasEXTExtendedDynamicState1 &&
-			OptionalDeviceExtensions.HasEXTExtendedDynamicState2 &&
-			OptionalDeviceExtensions.HasEXTExtendedDynamicState3 &&
-			// use a single vertex input stage for all pipelines
-			OptionalDeviceExtensions.HasEXTVertexInputDynamicState &&
 			// dynamic rendering removes need for renderpass at pipeline creation
 			OptionalDeviceExtensions.HasKHRDynamicRendering &&  
 			// in bindless a single pipeline layout is needed
diff --git a/Engine/Source/Runtime/VulkanRHI/Private/VulkanExtensions.cpp b/Engine/Source/Runtime/VulkanRHI/Private/VulkanExtensions.cpp
--- a/Engine/Source/Runtime/VulkanRHI/Private/VulkanExtensions.cpp
+++ b/Engine/Source/Runtime/VulkanRHI/Private/VulkanExtensions.cpp
@@ -2135,9 +2135,7 @@ public:
 	FVulkanEXTShaderObjectExtension(FVulkanDevice* InDevice)
 		: FVulkanDeviceExtension(InDevice, VK_EXT_SHADER_OBJECT_EXTENSION_NAME, VULKAN_EXTENSION_ENABLED)
 	{
-		bEnabledInCode = bEnabledInCode && 
-			(GVulkanAllowDynamicStatesCVar.GetValueOnAnyThread() != 0) &&
-			(GVulkanAllowShaderObjectCVar.GetValueOnAnyThread() != 0);
+		bEnabledInCode = bEnabledInCode && (GVulkanAllowShaderObjectCVar.GetValueOnAnyThread() != 0);
 	}
 
 	virtual void PrePhysicalDeviceFeatures(VkPhysicalDeviceFeatures2KHR& PhysicalDeviceFeatures2) override final
diff --git a/Engine/Source/Runtime/VulkanRHI/Private/VulkanPendingState.cpp b/Engine/Source/Runtime/VulkanRHI/Private/VulkanPendingState.cpp
--- a/Engine/Source/Runtime/VulkanRHI/Private/VulkanPendingState.cpp
+++ b/Engine/Source/Runtime/VulkanRHI/Private/VulkanPendingState.cpp
@@ -497,6 +497,11 @@ void FVulkanPendingGfxState::PrepareForDraw(FVulkanCommandListContext& Context)
 
 void FVulkanPendingGfxState::UpdateDynamicStates(FVulkanCommandBuffer& CommandBuffer)
 {
+	const bool bHasEXTExtendedDynamicState1 = Device.GetOptionalExtensions().HasEXTExtendedDynamicState1 || Device.SupportsShaderObjects();
+	const bool bHasEXTExtendedDynamicState2 = Device.GetOptionalExtensions().HasEXTExtendedDynamicState2 || Device.SupportsShaderObjects();
+	const bool bHasEXTExtendedDynamicState3 = Device.GetOptionalExtensions().HasEXTExtendedDynamicState3 || Device.SupportsShaderObjects();
+	const bool bHasEXTVertexInputDynamicState = Device.GetOptionalExtensions().HasEXTVertexInputDynamicState || Device.SupportsShaderObjects();
+
 	const bool bNeedsUpdateViewport = !CommandBuffer.bHasViewport || Viewports.Num() != CommandBuffer.CurrentViewports.Num() ||
 		(FMemory::Memcmp((const void*)CommandBuffer.CurrentViewports.GetData(), (const void*)Viewports.GetData(), Viewports.Num() * sizeof(VkViewport)) != 0);
 	// Validate and update Viewport
@@ -513,7 +518,7 @@ void FVulkanPendingGfxState::UpdateDynamicStates(FVulkanCommandBuffer& CommandBu
 			FlippedViewport.height = -FlippedViewport.height;
 		}
 
-		if (Device.GetOptionalExtensions().HasEXTExtendedDynamicState1)
+		if (bHasEXTExtendedDynamicState1)
 		{
 			VulkanRHI::vkCmdSetViewportWithCountEXT(CommandBuffer.GetHandle(), FlippedViewports.Num(), FlippedViewports.GetData());
 		}
@@ -530,7 +535,7 @@ void FVulkanPendingGfxState::UpdateDynamicStates(FVulkanCommandBuffer& CommandBu
 		(FMemory::Memcmp((const void*)CommandBuffer.CurrentScissors.GetData(), (const void*)Scissors.GetData(), Scissors.Num() * sizeof(VkRect2D)) != 0);
 	if (bNeedsUpdateScissor)
 	{
-		if (Device.GetOptionalExtensions().HasEXTExtendedDynamicState1)
+		if (bHasEXTExtendedDynamicState1)
 		{
 			VulkanRHI::vkCmdSetScissorWithCountEXT(CommandBuffer.GetHandle(), Scissors.Num(), Scissors.GetData());
 		}
@@ -573,7 +578,7 @@ void FVulkanPendingGfxState::UpdateDynamicStates(FVulkanCommandBuffer& CommandBu
 
 	// NOTE: All of the different VK_DYNAMIC_STATE_* values being set can be found in VulkanPipeline.cpp
 
-	if (Device.GetOptionalExtensions().HasEXTExtendedDynamicState1)
+	if (bHasEXTExtendedDynamicState1)
 	{
 		check(LastDynamicStateUpdate);
 
@@ -645,7 +650,7 @@ void FVulkanPendingGfxState::UpdateDynamicStates(FVulkanCommandBuffer& CommandBu
 		VULKAN_DYNAMIC_STATE_BLOCK_END();
 	}
 
-	if (Device.GetOptionalExtensions().HasEXTExtendedDynamicState2)
+	if (bHasEXTExtendedDynamicState2)
 	{
 		// VK_DYNAMIC_STATE_DEPTH_BIAS_ENABLE
 		const float OldDepthBiasConstantFactor = NewDynamicState.Rasterizer.DepthBiasConstantFactor;
@@ -653,14 +658,14 @@ void FVulkanPendingGfxState::UpdateDynamicStates(FVulkanCommandBuffer& CommandBu
 		VulkanRHI::vkCmdSetDepthBiasEnableEXT(CommandBuffer.GetHandle(), (NewDynamicState.Rasterizer.DepthBiasConstantFactor != 0.0f) ? VK_TRUE : VK_FALSE);
 		// Put back the old value if VK_DYNAMIC_STATE_DEPTH_BIAS is int use so that it gets updated
 		// Note: this does mean we call enable for each DepthBiasConstantFactor
-		if (Device.GetOptionalExtensions().HasEXTExtendedDynamicState3)
+		if (bHasEXTExtendedDynamicState3)
 		{
 			LastDynamicStateUpdate->Rasterizer.DepthBiasConstantFactor = OldDepthBiasConstantFactor;
 		}
 		VULKAN_DYNAMIC_STATE_BLOCK_END();
 	}
 
-	if (Device.GetOptionalExtensions().HasEXTExtendedDynamicState3)
+	if (bHasEXTExtendedDynamicState3)
 	{
 		// VK_DYNAMIC_STATE_POLYGON_MODE_EXT
 		VULKAN_DYNAMIC_STATE_BLOCK_START(Rasterizer.PolygonMode);
@@ -739,7 +744,7 @@ void FVulkanPendingGfxState::UpdateDynamicStates(FVulkanCommandBuffer& CommandBu
 		}
 	}
 
-	if (Device.GetOptionalExtensions().HasEXTVertexInputDynamicState)
+	if (bHasEXTVertexInputDynamicState)
 	{
 		bDirtyDynamicVertexInput = bDirtyDynamicVertexInput || bForceUpdate;
 
