VulkanRHI: Add shader object support (WIP).

#rnx

diff --git a/Engine/Source/Runtime/VulkanRHI/Private/VulkanDevice.cpp b/Engine/Source/Runtime/VulkanRHI/Private/VulkanDevice.cpp
--- a/Engine/Source/Runtime/VulkanRHI/Private/VulkanDevice.cpp
+++ b/Engine/Source/Runtime/VulkanRHI/Private/VulkanDevice.cpp
@@ -21,6 +21,7 @@
 #include "VulkanDescriptorSets.h"
 #include "VulkanChunkedPipelineCache.h"
 #include "VulkanBindlessDescriptorManager.h"
+#include "VulkanShaderObjectManager.h"
 
 static TAutoConsoleVariable<int32> GRHIAllowAsyncComputeCvar(
 	TEXT("r.Vulkan.AllowAsyncCompute"),
@@ -1396,6 +1397,8 @@ void FVulkanDevice::InitGPU()
 	BindlessDescriptorManager = new FVulkanBindlessDescriptorManager(*this);
 	BindlessDescriptorManager->Init();
 
+	ShaderObjectManager = new FVulkanShaderObjectManager(*this);
+
 	InitGlobalSamplers();
 
 	if (UE::RHI::UseGPUCrashBreadcrumbs())
@@ -1493,6 +1496,12 @@ void FVulkanDevice::Destroy()
 
 	DeferredDeletionQueue.Clear();
 
+	if (ShaderObjectManager)
+	{
+		delete ShaderObjectManager;
+		ShaderObjectManager = nullptr;
+	}
+
 	BindlessDescriptorManager->Deinit();
 	delete BindlessDescriptorManager;
 	BindlessDescriptorManager = nullptr;
diff --git a/Engine/Source/Runtime/VulkanRHI/Private/VulkanDevice.h b/Engine/Source/Runtime/VulkanRHI/Private/VulkanDevice.h
--- a/Engine/Source/Runtime/VulkanRHI/Private/VulkanDevice.h
+++ b/Engine/Source/Runtime/VulkanRHI/Private/VulkanDevice.h
@@ -27,6 +27,7 @@ class FVulkanRayTracingCompactionRequestHandler;
 class FVulkanRayTracingPipelineLibraryCache;
 struct FVulkanParallelRenderPassInfo;
 class FVulkanFenceManager;
+class FVulkanShaderObjectManager;
 
 #define VULKAN_USE_DEBUG_NAMES 1
 #if VULKAN_USE_DEBUG_NAMES
@@ -98,7 +99,6 @@ struct FOptionalVulkanDeviceExtensions
 			// Promoted to 1.3
 			uint64 HasEXTTextureCompressionASTCHDR : 1;
 			uint64 HasKHRMaintenance4 : 1;
-			uint64 HasKHRMaintenance5 : 1;
 			uint64 HasKHRSynchronization2 : 1;
 			uint64 HasKHRDynamicRendering : 1;
 			uint64 HasEXTSubgroupSizeControl : 1;
@@ -109,6 +109,7 @@ struct FOptionalVulkanDeviceExtensions
 			uint64 HasEXTVertexInputDynamicState : 1;
 
 			// Promoted to 1.4
+			uint64 HasKHRMaintenance5 : 1;
 			uint64 HasEXTLoadStoreOpNone : 1;
 			uint64 HasEXTHostImageCopy : 1;
 		};
@@ -576,6 +577,11 @@ public:
 			SupportsBindless();
 	}
 
+	FVulkanShaderObjectManager* GetShaderObjectManager()
+	{
+		return ShaderObjectManager;
+	}
+
 	bool SupportsShaderObjects() const
 	{
 		// NOTE: No need to check individual dynamic state extension, they all get enabled by turning on shader objects
@@ -659,6 +665,7 @@ private:
 	FVulkanBindlessDescriptorManager* BindlessDescriptorManager = nullptr;
 
 	FVulkanShaderFactory ShaderFactory;
+	FVulkanShaderObjectManager* ShaderObjectManager = nullptr;
 
 	VkPhysicalDevice Gpu;
 	VkPhysicalDeviceProperties GpuProps;
diff --git a/Engine/Source/Runtime/VulkanRHI/Private/VulkanPipeline.cpp b/Engine/Source/Runtime/VulkanRHI/Private/VulkanPipeline.cpp
--- a/Engine/Source/Runtime/VulkanRHI/Private/VulkanPipeline.cpp
+++ b/Engine/Source/Runtime/VulkanRHI/Private/VulkanPipeline.cpp
@@ -21,6 +21,7 @@
 #include "VulkanRenderTargetLayout.h"
 #include "VulkanRenderpass.h"
 #include "VulkanBindlessDescriptorManager.h"
+#include "VulkanShaderObjectManager.h"
 
 #define LRU_DEBUG 0
 #if !UE_BUILD_SHIPPING
@@ -1635,6 +1636,12 @@ bool FVulkanPipelineStateCacheManager::CreateGfxPipelineFromLibrary(FVulkanGraph
 
 bool FVulkanPipelineStateCacheManager::CreateGfxPipelineFromEntry(FVulkanGraphicsPipelineState* PSO, FVulkanShader* Shaders[ShaderStage::NumGraphicsStages], FGraphicsPipelineStateInitializer::EPSOPrecacheCompileType PSOCompileType)
 {
+	// When shader objects are used, create linked shader objects instead of pipelines
+	if (Device->SupportsShaderObjects())
+	{
+		return Device->GetShaderObjectManager()->CreateLinkedShaderObjects(PSO, Shaders, PSOCompileType);
+	}
+
 	// See if it can be created from a pipeline library first
 	if (Device->SupportsGraphicPipelineLibraries() && CreateGfxPipelineFromLibrary(*PSO, Shaders, PSOCompileType))
 	{
@@ -2755,6 +2762,12 @@ FVulkanComputePipeline* FVulkanPipelineStateCacheManager::CreateComputePipelineF
 	DescriptorSetLayoutInfo.FinalizeBindings<true>(*Device, UBGatherInfo, TArrayView<FRHISamplerState*>(), Shader->UsesBindless());
 	FVulkanLayout* Layout = FindOrAddLayout(DescriptorSetLayoutInfo, false, Shader->UsesBindless());
 	checkSlow(!Layout->IsGfxLayout());
+	Pipeline->Layout = Layout;
+
+	if (Device->SupportsShaderObjects())
+	{
+		return Pipeline;
+	}
 
 	TRefCountPtr<FVulkanShaderModule> ShaderModule = Shader->GetOrCreateHandle(Layout, Layout->GetDescriptorSetLayoutHash());
 
@@ -2805,7 +2818,6 @@ FVulkanComputePipeline* FVulkanPipelineStateCacheManager::CreateComputePipelineF
 		Pipeline->SetValid(false);
 	}
 
-	Pipeline->Layout = Layout;
 
 	INC_DWORD_STAT(STAT_VulkanNumPSOs);
 
@@ -3116,6 +3128,31 @@ void FVulkanGraphicsPipelineState::DeleteVkPipeline(bool bImmediate)
 	Device->PipelineStateCache->LRUCheckNotInside(this);
 }
 
+void FVulkanGraphicsPipelineState::Bind(VkCommandBuffer CmdBuffer)
+{
+	if (Device->SupportsShaderObjects())
+	{
+		Device->GetShaderObjectManager()->BindGraphicsShaderObjects(CmdBuffer, VulkanShaders);
+	}
+	else
+	{
+		VulkanRHI::vkCmdBindPipeline(CmdBuffer, VK_PIPELINE_BIND_POINT_GRAPHICS, VulkanPipeline);
+	}
+}
+
+void FVulkanComputePipeline::Bind(VkCommandBuffer CmdBuffer)
+{
+	if (Device->SupportsShaderObjects())
+	{
+		check(ComputeShader.IsValid());
+		Device->GetShaderObjectManager()->BindComputeShaderObject(CmdBuffer, ResourceCast(GetComputeShader()));
+	}
+	else
+	{
+		VulkanRHI::vkCmdBindPipeline(CmdBuffer, VK_PIPELINE_BIND_POINT_COMPUTE, Pipeline);
+	}
+}
+
 void FVulkanPipelineStateCacheManager::LRUCheckNotInside(FVulkanGraphicsPipelineState* PSO)
 {
 	FScopeLock Lock(&LRUCS);
diff --git a/Engine/Source/Runtime/VulkanRHI/Private/VulkanPipeline.h b/Engine/Source/Runtime/VulkanRHI/Private/VulkanPipeline.h
--- a/Engine/Source/Runtime/VulkanRHI/Private/VulkanPipeline.h
+++ b/Engine/Source/Runtime/VulkanRHI/Private/VulkanPipeline.h
@@ -717,10 +717,7 @@ public:
 		return *Layout;
 	}
 
-	void Bind(VkCommandBuffer CmdBuffer)
-	{
-		VulkanRHI::vkCmdBindPipeline(CmdBuffer, VK_PIPELINE_BIND_POINT_COMPUTE, Pipeline);
-	}
+	void Bind(VkCommandBuffer CmdBuffer);
 
 	bool UsesBindless() const
 	{
@@ -737,34 +734,31 @@ public:
 	FVulkanGraphicsPipelineState(FVulkanDevice* Device, const FGraphicsPipelineStateInitializer& PSOInitializer, const FGfxPipelineDesc& Desc, FVulkanPSOKey* Key);
 	~FVulkanGraphicsPipelineState();
 
-	inline const FVulkanVertexInputStateInfo& GetVertexInputState() const
+	const FVulkanVertexInputStateInfo& GetVertexInputState() const
 	{
 		return VertexInputState;
 	}
 
-	inline const FVulkanLayout& GetLayout() const
+	const FVulkanLayout& GetLayout() const
 	{
 		return *Layout;
 	}
 
-	inline void Bind(VkCommandBuffer CmdBuffer)
-	{
-		VulkanRHI::vkCmdBindPipeline(CmdBuffer, VK_PIPELINE_BIND_POINT_GRAPHICS, VulkanPipeline);
-	}
+	void Bind(VkCommandBuffer CmdBuffer);
 
-	inline const uint64 GetShaderKey(EShaderFrequency Frequency) const
+	const uint64 GetShaderKey(EShaderFrequency Frequency) const
 	{
 		ShaderStage::EStage Stage = ShaderStage::GetStageForFrequency(Frequency);
 		return ShaderKeys[Stage];
 	}
 
-	inline const FVulkanShader* GetVulkanShader(EShaderFrequency Frequency) const
+	const FVulkanShader* GetVulkanShader(EShaderFrequency Frequency) const
 	{
 		ShaderStage::EStage Stage = ShaderStage::GetStageForFrequency(Frequency);
 		return VulkanShaders[Stage];
 	}
 
-	inline VkPipeline GetVulkanPipeline() const
+	VkPipeline GetVulkanPipeline() const
 	{
 		return VulkanPipeline;
 	}
diff --git a/Engine/Source/Runtime/VulkanRHI/Private/VulkanShaderObjectManager.cpp b/Engine/Source/Runtime/VulkanRHI/Private/VulkanShaderObjectManager.cpp
new file mode 100644
--- /dev/null
+++ b/Engine/Source/Runtime/VulkanRHI/Private/VulkanShaderObjectManager.cpp
@@ -0,0 +1,148 @@
+// Copyright Epic Games, Inc. All Rights Reserved.
+
+#include "VulkanRHIPrivate.h"
+#include "VulkanShaderObjectManager.h"
+#include "VulkanResources.h"
+#include "VulkanDevice.h"
+#include "VulkanBindlessDescriptorManager.h"
+
+
+FVulkanShaderObjectManager::FVulkanShaderObjectManager(FVulkanDevice& InDevice)
+	: Device(InDevice)
+{
+}
+
+VkShaderStageFlags FVulkanShaderObjectManager::GetPossibleNextStages(VkShaderStageFlagBits Stage)
+{
+	switch (Stage)
+	{
+	case VK_SHADER_STAGE_VERTEX_BIT:
+		return Device.GetPhysicalDeviceFeatures().Core_1_0.geometryShader ?
+			VK_SHADER_STAGE_GEOMETRY_BIT | VK_SHADER_STAGE_FRAGMENT_BIT : VK_SHADER_STAGE_FRAGMENT_BIT;
+
+	case VK_SHADER_STAGE_GEOMETRY_BIT: 
+	case VK_SHADER_STAGE_MESH_BIT_EXT:
+		return VK_SHADER_STAGE_FRAGMENT_BIT;
+
+	case VK_SHADER_STAGE_TASK_BIT_EXT:
+		return VK_SHADER_STAGE_MESH_BIT_EXT;
+
+	case VK_SHADER_STAGE_FRAGMENT_BIT:
+	case VK_SHADER_STAGE_COMPUTE_BIT:
+		return 0;
+	}
+
+	checkNoEntry();
+	return 0;
+}
+
+VkShaderEXT FVulkanShaderObjectManager::CreateShaderObject(FVulkanShader& VulkanShader)
+{
+	checkf(VulkanShader.UsesBindless(), TEXT("Shader Objects currently only supported in bindless."));
+
+	const VkShaderStageFlagBits Stage = UEFrequencyToVKStageBit(VulkanShader.Frequency);
+	FVulkanShader::FSpirvCode SpirvCode = VulkanShader.GetSpirvCode();
+
+	ANSICHAR EntryPoint[24];
+	VulkanShader.GetEntryPoint(EntryPoint, 24);
+
+	TConstArrayView<VkDescriptorSetLayout> DescriptorSetLayout = Device.GetBindlessDescriptorManager()->GetDescriptorSetLayouts();
+
+	VkShaderCreateInfoEXT ShaderCreateInfo = {
+		.sType = VK_STRUCTURE_TYPE_SHADER_CREATE_INFO_EXT,
+		.pNext = nullptr,
+		.flags = 0,
+		.stage = Stage,
+		.nextStage = GetPossibleNextStages(Stage),
+		.codeType = VK_SHADER_CODE_TYPE_SPIRV_EXT,
+		.codeSize = (size_t)SpirvCode.GetCodeSizeInBytes(),
+		.pCode = SpirvCode.GetCodeView().GetData(),
+		.pName = EntryPoint,
+		.setLayoutCount = (uint32)DescriptorSetLayout.Num(),
+		.pSetLayouts = DescriptorSetLayout.GetData(),
+		.pushConstantRangeCount = 0,
+		.pPushConstantRanges = nullptr,
+		.pSpecializationInfo = nullptr
+	};
+
+	// :todo-jn: Task shaders never used at the moment
+	if (Stage == VK_SHADER_STAGE_MESH_BIT_EXT)
+	{
+		ShaderCreateInfo.flags |= VK_SHADER_CREATE_NO_TASK_SHADER_BIT_EXT;
+	}
+
+	// :todo-jn: Could have a penalty, should narrow when it's applied
+	if (GRHISupportsAttachmentVariableRateShading && (Stage == VK_SHADER_STAGE_FRAGMENT_BIT))
+	{
+		if (ValidateShadingRateDataType())
+		{
+			if (GRHIVariableRateShadingImageDataType == VRSImage_Palette)
+			{
+				ShaderCreateInfo.flags |= VK_SHADER_CREATE_FRAGMENT_SHADING_RATE_ATTACHMENT_BIT_EXT;
+			}
+			if (GRHIVariableRateShadingImageDataType == VRSImage_Fractional)
+			{
+				ShaderCreateInfo.flags |= VK_SHADER_CREATE_FRAGMENT_DENSITY_MAP_ATTACHMENT_BIT_EXT;
+			}
+		}
+	}
+
+	VkShaderEXT ShaderObjectHandle = VK_NULL_HANDLE;
+	VERIFYVULKANRESULT(VulkanRHI::vkCreateShadersEXT(Device.GetHandle(), 1, &ShaderCreateInfo, VULKAN_CPU_ALLOCATOR, &ShaderObjectHandle));
+	return ShaderObjectHandle;
+}
+
+bool FVulkanShaderObjectManager::CreateLinkedShaderObjects(FVulkanGraphicsPipelineState* PSO, FVulkanShader* Shaders[ShaderStage::NumGraphicsStages], FGraphicsPipelineStateInitializer::EPSOPrecacheCompileType PSOCompileType)
+{
+	// :todo-jn: Create new linked binaries from SPIRV with VK_SHADER_CREATE_LINK_STAGE_BIT_EXT
+	// :todo-jn: Calls to bind will intercept and replace with linked binaries when present.
+	// :todo-jn: These are the binaries will be cached.
+	return true;
+}
+
+void FVulkanShaderObjectManager::BindGraphicsShaderObjects(VkCommandBuffer CommandBufferHandle, TArrayView<FVulkanShader*> VulkanShaders)
+{
+	VkShaderStageFlagBits ShaderStages[ShaderStage::NumGraphicsStages];
+	VkShaderEXT ShaderObjectHandles[ShaderStage::NumGraphicsStages] = { VK_NULL_HANDLE };
+
+	uint32 StageCount = 0;
+	auto AddStage = [&StageCount, &ShaderStages, &ShaderObjectHandles, &VulkanShaders](ShaderStage::EStage Stage, VkShaderStageFlagBits StageFlag)
+	{
+		FVulkanShader* VulkanShader = VulkanShaders[Stage];
+		ShaderStages[StageCount] = StageFlag;
+		ShaderObjectHandles[StageCount] = VulkanShader ? VulkanShader->ShaderObject : VK_NULL_HANDLE;
+		check(!VulkanShader || (Stage == ShaderStage::GetStageForFrequency(VulkanShader->Frequency)));
+		StageCount++;
+	};
+
+	AddStage(ShaderStage::Vertex, VK_SHADER_STAGE_VERTEX_BIT);
+	AddStage(ShaderStage::Pixel, VK_SHADER_STAGE_FRAGMENT_BIT);
+
+#if PLATFORM_SUPPORTS_GEOMETRY_SHADERS
+	if (Device.GetPhysicalDeviceFeatures().Core_1_0.geometryShader)
+	{
+		AddStage(ShaderStage::Geometry, VK_SHADER_STAGE_GEOMETRY_BIT);
+	}
+#endif
+
+#if PLATFORM_SUPPORTS_MESH_SHADERS
+	if (Device.GetOptionalExtensions().HasEXTMeshShader)
+	{
+		AddStage(ShaderStage::Mesh, VK_SHADER_STAGE_MESH_BIT_EXT);
+		AddStage(ShaderStage::Task, VK_SHADER_STAGE_TASK_BIT_EXT);
+	}
+#endif
+
+	// :todo-jn: Cache state and avoid setting the same shader over and over
+	VulkanRHI::vkCmdBindShadersEXT(CommandBufferHandle, StageCount, ShaderStages, ShaderObjectHandles);
+}
+
+void FVulkanShaderObjectManager::BindComputeShaderObject(VkCommandBuffer CommandBufferHandle, FVulkanShader* VulkanShader)
+{
+	check(VulkanShader && VulkanShader->ShaderObject);
+
+	static const VkShaderStageFlagBits ShaderStage = VK_SHADER_STAGE_COMPUTE_BIT;
+	const VkShaderEXT ShaderObjectHandle = VulkanShader->ShaderObject;
+
+	VulkanRHI::vkCmdBindShadersEXT(CommandBufferHandle, 1, &ShaderStage, &ShaderObjectHandle);
+}
diff --git a/Engine/Source/Runtime/VulkanRHI/Private/VulkanShaderObjectManager.h b/Engine/Source/Runtime/VulkanRHI/Private/VulkanShaderObjectManager.h
new file mode 100644
--- /dev/null
+++ b/Engine/Source/Runtime/VulkanRHI/Private/VulkanShaderObjectManager.h
@@ -0,0 +1,29 @@
+// Copyright Epic Games, Inc. All Rights Reserved.
+
+#pragma once
+
+#include "VulkanPlatform.h"
+#include "VulkanCommon.h"
+#include "RHIResources.h"
+
+class FVulkanDevice;
+class FVulkanShader;
+class FVulkanGraphicsPipelineState;
+
+
+class FVulkanShaderObjectManager
+{
+public:
+	FVulkanShaderObjectManager(FVulkanDevice& InDevice);
+
+	VkShaderEXT CreateShaderObject(FVulkanShader& VulkanShader);
+	bool CreateLinkedShaderObjects(FVulkanGraphicsPipelineState* PSO, FVulkanShader* Shaders[ShaderStage::NumGraphicsStages], FGraphicsPipelineStateInitializer::EPSOPrecacheCompileType PSOCompileType);
+
+	void BindGraphicsShaderObjects(VkCommandBuffer CommandBufferHandle, TArrayView<FVulkanShader*> VulkanShaders);
+	void BindComputeShaderObject(VkCommandBuffer CommandBufferHandle, FVulkanShader* VulkanShader);
+
+private:
+	VkShaderStageFlags GetPossibleNextStages(VkShaderStageFlagBits Stage);
+
+	FVulkanDevice& Device;
+};
\ No newline at end of file
diff --git a/Engine/Source/Runtime/VulkanRHI/Private/VulkanShaders.cpp b/Engine/Source/Runtime/VulkanRHI/Private/VulkanShaders.cpp
--- a/Engine/Source/Runtime/VulkanRHI/Private/VulkanShaders.cpp
+++ b/Engine/Source/Runtime/VulkanRHI/Private/VulkanShaders.cpp
@@ -13,6 +13,7 @@
 #include "VulkanDescriptorSets.h"
 #include "RHICoreShader.h"
 #include "VulkanBindlessDescriptorManager.h"
+#include "VulkanShaderObjectManager.h"
 
 
 extern int32 GVulkanGraphicPipelineLibraryLinkingMode;
@@ -79,9 +80,9 @@ static VkPipeline CreatePreRasterShaderLibrary(FVulkanDevice* Device, FVulkanSha
 	VkShaderModuleCreateInfo ModuleCreateInfo;
 	ZeroVulkanStruct(ModuleCreateInfo, VK_STRUCTURE_TYPE_SHADER_MODULE_CREATE_INFO);
 
-	TArrayView<uint32> SpirvCode = Shader->GetSpirvCode().GetCodeView();
-	ModuleCreateInfo.codeSize = SpirvCode.Num() * sizeof(uint32);
-	ModuleCreateInfo.pCode = SpirvCode.GetData();
+	FVulkanShader::FSpirvCode SpirvCode = Shader->GetSpirvCode();
+	ModuleCreateInfo.codeSize = SpirvCode.GetCodeSizeInBytes();
+	ModuleCreateInfo.pCode = SpirvCode.GetCodeView().GetData();
 
 	VkGraphicsPipelineLibraryCreateInfoEXT LibraryInfo;
 	ZeroVulkanStruct(LibraryInfo, VK_STRUCTURE_TYPE_GRAPHICS_PIPELINE_LIBRARY_CREATE_INFO_EXT);
@@ -143,9 +144,9 @@ static VkPipeline CreatePixelShaderLibrary(FVulkanDevice* Device, FVulkanShader*
 	VkShaderModuleCreateInfo ModuleCreateInfo;
 	ZeroVulkanStruct(ModuleCreateInfo, VK_STRUCTURE_TYPE_SHADER_MODULE_CREATE_INFO);
 
-	TArrayView<uint32> SpirvCode = PixelShader->GetSpirvCode().GetCodeView();
-	ModuleCreateInfo.codeSize = SpirvCode.Num() * sizeof(uint32);
-	ModuleCreateInfo.pCode = SpirvCode.GetData();
+	FVulkanShader::FSpirvCode SpirvCode = PixelShader->GetSpirvCode();
+	ModuleCreateInfo.codeSize = SpirvCode.GetCodeSizeInBytes();
+	ModuleCreateInfo.pCode = SpirvCode.GetCodeView().GetData();
 
 	VkGraphicsPipelineLibraryCreateInfoEXT LibraryInfo;
 	ZeroVulkanStruct(LibraryInfo, VK_STRUCTURE_TYPE_GRAPHICS_PIPELINE_LIBRARY_CREATE_INFO_EXT);
@@ -246,7 +247,11 @@ ShaderType* FVulkanShaderFactory::CreateShader(TArrayView<const uint8> Code, FVu
 				FShaderCodeReader ShaderCode(Code);
 				ReadShaderOptionalData(ShaderCode, RetShader);
 
-				if (Device->SupportsGraphicPipelineLibraries())
+				if (Device->SupportsShaderObjects())
+				{
+					RetShader->ShaderObject = Device->GetShaderObjectManager()->CreateShaderObject(*RetShader);
+				}
+				else if (Device->SupportsGraphicPipelineLibraries())
 				{
 					if (std::is_same_v<ShaderType, FVulkanVertexShader>)
 					{
@@ -733,6 +738,15 @@ FVulkanShader::~FVulkanShader()
 	PurgeShaderModules();
 	Device->GetShaderFactory().OnDeleteShader(*this);
 
+	if (ShaderObject)
+	{
+		Device->GetDeferredDeletionQueue().EnqueueFunc([Handle=ShaderObject](FVulkanDevice& Device) 
+		{
+			VulkanRHI::vkDestroyShaderEXT(Device.GetHandle(), Handle, VULKAN_CPU_ALLOCATOR);
+		});
+		ShaderObject = VK_NULL_HANDLE;
+	}
+
 	if (PipelineLibrary)
 	{
 		Device->GetDeferredDeletionQueue().EnqueueResource(VulkanRHI::FDeferredDeletionQueue2::EType::Pipeline, PipelineLibrary);
