VulkanRHI: Add support for Maintenance6 extension.

#rnx

diff --git a/Engine/Source/Runtime/VulkanRHI/Private/VulkanCommandWrappers.h b/Engine/Source/Runtime/VulkanRHI/Private/VulkanCommandWrappers.h
--- a/Engine/Source/Runtime/VulkanRHI/Private/VulkanCommandWrappers.h
+++ b/Engine/Source/Runtime/VulkanRHI/Private/VulkanCommandWrappers.h
@@ -249,6 +249,7 @@ struct FWrapLayer
 	static void GetDescriptorSetLayoutBindingOffsetEXT(VkResult Result, VkDevice Device, VkDescriptorSetLayout Layout, uint32_t Binding, VkDeviceSize* Offset) VULKAN_LAYER_BODY
 	static void CmdBindDescriptorBuffersEXT(VkResult Result, VkCommandBuffer CommandBuffer, uint32_t BufferCount, const VkDescriptorBufferBindingInfoEXT* BindingInfos) VULKAN_LAYER_BODY
 	static void CmdSetDescriptorBufferOffsetsEXT(VkResult Result, VkCommandBuffer CommandBuffer, VkPipelineBindPoint PipelineBindPoint, VkPipelineLayout Layout, uint32_t FirstSet, uint32_t SetCount, const uint32_t* BufferIndices, const VkDeviceSize* Offsets) VULKAN_LAYER_BODY
+	static void CmdSetDescriptorBufferOffsets2EXT(VkResult Result, VkCommandBuffer CommandBuffer, const VkSetDescriptorBufferOffsetsInfoEXT* SetDescriptorBufferOffsetsInfo) VULKAN_LAYER_BODY
 	static void GetDescriptorEXT(VkResult Result, VkDevice Device, const VkDescriptorGetInfoEXT* DescriptorInfo, size_t DataSize, void* Descriptor) VULKAN_LAYER_BODY
 	static void CreateDeferredOperationKHR(VkResult Result, VkDevice Device, const VkAllocationCallbacks* Allocator, VkDeferredOperationKHR* DeferredOperation) VULKAN_LAYER_BODY
 	static void DestroyDeferredOperationKHR(VkResult Result, VkDevice Device, VkDeferredOperationKHR DeferredOperation, const VkAllocationCallbacks* Allocator) VULKAN_LAYER_BODY
@@ -1766,6 +1767,13 @@ namespace VulkanRHI
 		FWrapLayer::CmdSetDescriptorBufferOffsetsEXT(VK_SUCCESS, CommandBuffer, PipelineBindPoint, Layout, FirstSet, SetCount, BufferIndices, Offsets);
 	}
 
+	static FORCEINLINE_DEBUGGABLE void vkCmdSetDescriptorBufferOffsets2EXT(VkCommandBuffer CommandBuffer, const VkSetDescriptorBufferOffsetsInfoEXT* SetDescriptorBufferOffsetsInfo)
+	{
+		FWrapLayer::CmdSetDescriptorBufferOffsets2EXT(VK_RESULT_MAX_ENUM, CommandBuffer, SetDescriptorBufferOffsetsInfo);
+		VULKANAPINAMESPACE::vkCmdSetDescriptorBufferOffsets2EXT(CommandBuffer, SetDescriptorBufferOffsetsInfo);
+		FWrapLayer::CmdSetDescriptorBufferOffsets2EXT(VK_SUCCESS, CommandBuffer, SetDescriptorBufferOffsetsInfo);
+	}
+
 	static FORCEINLINE_DEBUGGABLE void vkGetDescriptorEXT(VkDevice Device, const VkDescriptorGetInfoEXT* DescriptorInfo, size_t DataSize, void* Descriptor)
 	{
 		FWrapLayer::GetDescriptorEXT(VK_RESULT_MAX_ENUM, Device, DescriptorInfo, DataSize, Descriptor);
diff --git a/Engine/Source/Runtime/VulkanRHI/Private/VulkanDebug.cpp b/Engine/Source/Runtime/VulkanRHI/Private/VulkanDebug.cpp
--- a/Engine/Source/Runtime/VulkanRHI/Private/VulkanDebug.cpp
+++ b/Engine/Source/Runtime/VulkanRHI/Private/VulkanDebug.cpp
@@ -4477,6 +4477,17 @@ void FWrapLayer::CmdSetDescriptorBufferOffsetsEXT(VkResult Result, VkCommandBuff
 	}
 }
 
+void FWrapLayer::CmdSetDescriptorBufferOffsets2EXT(VkResult Result, VkCommandBuffer CommandBuffer, const VkSetDescriptorBufferOffsetsInfoEXT* SetDescriptorBufferOffsetsInfo)
+{
+	if (Result == VK_RESULT_MAX_ENUM)
+	{
+#if VULKAN_ENABLE_DUMP_LAYER
+		PrintfBeginResult(FString::Printf(TEXT("CmdSetDescriptorBufferOffsets2EXT(CommandBuffer=0x%p, SetDescriptorBufferOffsetsInfo=0x%p)"), 
+			CommandBuffer, SetDescriptorBufferOffsetsInfo));
+#endif
+	}
+}
+
 void FWrapLayer::GetDescriptorEXT(VkResult Result, VkDevice Device, const VkDescriptorGetInfoEXT* DescriptorInfo, size_t DataSize, void* Descriptor)
 {
 	if (Result == VK_RESULT_MAX_ENUM)
diff --git a/Engine/Source/Runtime/VulkanRHI/Private/VulkanDevice.h b/Engine/Source/Runtime/VulkanRHI/Private/VulkanDevice.h
--- a/Engine/Source/Runtime/VulkanRHI/Private/VulkanDevice.h
+++ b/Engine/Source/Runtime/VulkanRHI/Private/VulkanDevice.h
@@ -110,6 +110,7 @@ struct FOptionalVulkanDeviceExtensions
 
 			// Promoted to 1.4
 			uint64 HasKHRMaintenance5 : 1;
+			uint64 HasKHRMaintenance6 : 1;
 			uint64 HasEXTLoadStoreOpNone : 1;
 			uint64 HasEXTHostImageCopy : 1;
 		};
diff --git a/Engine/Source/Runtime/VulkanRHI/Private/VulkanExtensions.cpp b/Engine/Source/Runtime/VulkanRHI/Private/VulkanExtensions.cpp
--- a/Engine/Source/Runtime/VulkanRHI/Private/VulkanExtensions.cpp
+++ b/Engine/Source/Runtime/VulkanRHI/Private/VulkanExtensions.cpp
@@ -288,7 +288,42 @@ private:
 	VkPhysicalDeviceMaintenance5FeaturesKHR Maintenance5Features;
 };
 
+// ***** VK_KHR_maintenance6
+class FVulkanKHRMaintenance6Extension : public FVulkanDeviceExtension
+{
+public:
 
+	FVulkanKHRMaintenance6Extension(FVulkanDevice* InDevice)
+		: FVulkanDeviceExtension(InDevice, VK_KHR_MAINTENANCE_6_EXTENSION_NAME, VULKAN_EXTENSION_ENABLED, VK_API_VERSION_1_4)
+	{
+	}
+
+	virtual void PrePhysicalDeviceFeatures(VkPhysicalDeviceFeatures2KHR& PhysicalDeviceFeatures2) override final
+	{
+		ZeroVulkanStruct(Maintenance6Features, VK_STRUCTURE_TYPE_PHYSICAL_DEVICE_MAINTENANCE_6_FEATURES_KHR);
+		AddToPNext(PhysicalDeviceFeatures2, Maintenance6Features);
+	}
+
+	virtual void PostPhysicalDeviceFeatures(FOptionalVulkanDeviceExtensions& ExtensionFlags) override final
+	{
+		bRequirementsPassed = (Maintenance6Features.maintenance6 == VK_TRUE);
+		if (bRequirementsPassed)
+		{
+			ExtensionFlags.HasKHRMaintenance6 = true;
+		}
+	}
+
+	virtual void PreCreateDevice(VkDeviceCreateInfo& DeviceCreateInfo) override final
+	{
+		if (bRequirementsPassed)
+		{
+			AddToPNext(DeviceCreateInfo, Maintenance6Features);
+		}
+	}
+
+private:
+	VkPhysicalDeviceMaintenance6FeaturesKHR Maintenance6Features;
+};
 // ***** VK_KHR_maintenance7
 class FVulkanKHRMaintenance7Extension : public FVulkanDeviceExtension
 {
diff --git a/Engine/Source/Runtime/VulkanRHI/Private/VulkanLoader.h b/Engine/Source/Runtime/VulkanRHI/Private/VulkanLoader.h
--- a/Engine/Source/Runtime/VulkanRHI/Private/VulkanLoader.h
+++ b/Engine/Source/Runtime/VulkanRHI/Private/VulkanLoader.h
@@ -222,6 +222,7 @@
 	EnumMacro(PFN_vkGetDescriptorSetLayoutBindingOffsetEXT, vkGetDescriptorSetLayoutBindingOffsetEXT) \
 	EnumMacro(PFN_vkCmdBindDescriptorBuffersEXT, vkCmdBindDescriptorBuffersEXT) \
 	EnumMacro(PFN_vkCmdSetDescriptorBufferOffsetsEXT, vkCmdSetDescriptorBufferOffsetsEXT) \
+	EnumMacro(PFN_vkCmdSetDescriptorBufferOffsets2EXT, vkCmdSetDescriptorBufferOffsets2EXT) \
 	EnumMacro(PFN_vkGetDescriptorEXT, vkGetDescriptorEXT) \
 	EnumMacro(PFN_vkCreateDeferredOperationKHR, vkCreateDeferredOperationKHR) \
 	EnumMacro(PFN_vkDestroyDeferredOperationKHR, vkDestroyDeferredOperationKHR) \
