VulkanRHI: Fix more dynamic states for shader objects (WIP).

#rnx

diff --git a/Engine/Source/Runtime/VulkanRHI/Private/VulkanCommandBuffer.cpp b/Engine/Source/Runtime/VulkanRHI/Private/VulkanCommandBuffer.cpp
--- a/Engine/Source/Runtime/VulkanRHI/Private/VulkanCommandBuffer.cpp
+++ b/Engine/Source/Runtime/VulkanRHI/Private/VulkanCommandBuffer.cpp
@@ -232,12 +232,10 @@ void FVulkanCommandBuffer::Begin(FVulkanQueryPool* OptionalQueryPool, VkRenderPa
 	}
 #endif
 
-	SetUnusedDynamicStates();
-
 	bNeedsFullDynamicStateUpdate = true;
 }
 
-void FVulkanCommandBuffer::SetUnusedDynamicStates()
+void FVulkanCommandBuffer::SetUnusedGraphicsDynamicStates()
 {
 	// Some unused states have been made dynamic with shader objects, set a default value at the start of each new command buffer
 	// :todo-jn: Look for a way to batch these settings (it was suggested to use a baked secondary command buffer)
@@ -251,8 +249,9 @@ void FVulkanCommandBuffer::SetUnusedDynamicStates()
 		VulkanRHI::vkCmdSetPrimitiveRestartEnableEXT(CommandBufferHandle, false);
 		VulkanRHI::vkCmdSetFrontFaceEXT(CommandBufferHandle, VK_FRONT_FACE_CLOCKWISE);
 		VulkanRHI::vkCmdSetLineWidth(CommandBufferHandle, 1.0f);
-		VulkanRHI::vkCmdSetStencilCompareMask(CommandBufferHandle, VK_STENCIL_FACE_FRONT_AND_BACK, MAX_uint32);
-		VulkanRHI::vkCmdSetStencilWriteMask(CommandBufferHandle, VK_STENCIL_FACE_FRONT_AND_BACK, MAX_uint32);
+
+		const float blendConstants[4] = {1.0f, 1.0f, 1.0f, 1.0f};
+		VulkanRHI::vkCmdSetBlendConstants(CommandBufferHandle, blendConstants);
 
 		const VkSampleMask SampleMask = MAX_uint32;
 		VulkanRHI::vkCmdSetSampleMaskEXT(CommandBufferHandle, VK_SAMPLE_COUNT_1_BIT, &SampleMask);
diff --git a/Engine/Source/Runtime/VulkanRHI/Private/VulkanCommandBuffer.h b/Engine/Source/Runtime/VulkanRHI/Private/VulkanCommandBuffer.h
--- a/Engine/Source/Runtime/VulkanRHI/Private/VulkanCommandBuffer.h
+++ b/Engine/Source/Runtime/VulkanRHI/Private/VulkanCommandBuffer.h
@@ -175,9 +175,9 @@ public:
 	// Public immutable reference to be accessed directly
 	FVulkanDevice& Device;
 
-private:
-	void SetUnusedDynamicStates();
+	void SetUnusedGraphicsDynamicStates();
 
+private:
 	VkCommandBuffer CommandBufferHandle = VK_NULL_HANDLE;
 	double SubmittedTime = 0.0f;
 
diff --git a/Engine/Source/Runtime/VulkanRHI/Private/VulkanPendingState.cpp b/Engine/Source/Runtime/VulkanRHI/Private/VulkanPendingState.cpp
--- a/Engine/Source/Runtime/VulkanRHI/Private/VulkanPendingState.cpp
+++ b/Engine/Source/Runtime/VulkanRHI/Private/VulkanPendingState.cpp
@@ -562,6 +562,12 @@ void FVulkanPendingGfxState::UpdateDynamicStates(FVulkanCommandBuffer& CommandBu
 	const bool bHasMeshShader = (CurrentPipeline->GetVulkanShader(SF_Mesh) != nullptr);
 	const bool bHasGeometryShader = (CurrentPipeline->GetVulkanShader(SF_Geometry) != nullptr);
 
+	// Set the unused graphic states at least once when we use a command buffer for graphics the first time
+	if (bForceUpdate)
+	{
+		CommandBuffer.SetUnusedGraphicsDynamicStates();
+	}
+
 	// Bunch of macros to make dynamic state comparions below easier to read
 #define VULKAN_DYNAMIC_STATE_COMPARE(Name) (LastDynamicStateUpdate->Name != NewDynamicState.Name)
 #define VULKAN_DYNAMIC_STATE_COMPARE1(Name, ...) VULKAN_DYNAMIC_STATE_COMPARE(Name)
@@ -653,15 +659,11 @@ void FVulkanPendingGfxState::UpdateDynamicStates(FVulkanCommandBuffer& CommandBu
 	if (bHasEXTExtendedDynamicState2)
 	{
 		// VK_DYNAMIC_STATE_DEPTH_BIAS_ENABLE
-		const float OldDepthBiasConstantFactor = NewDynamicState.Rasterizer.DepthBiasConstantFactor;
-		VULKAN_DYNAMIC_STATE_BLOCK_START(Rasterizer.DepthBiasConstantFactor);
+		// VK_DYNAMIC_STATE_DEPTH_BIAS
+		VULKAN_DYNAMIC_STATE_BLOCK_START(Rasterizer.DepthBiasConstantFactor, Rasterizer.DepthBiasSlopeScale);
+		constexpr float depthBiasClamp = 0.0f;
 		VulkanRHI::vkCmdSetDepthBiasEnableEXT(CommandBuffer.GetHandle(), (NewDynamicState.Rasterizer.DepthBiasConstantFactor != 0.0f) ? VK_TRUE : VK_FALSE);
-		// Put back the old value if VK_DYNAMIC_STATE_DEPTH_BIAS is int use so that it gets updated
-		// Note: this does mean we call enable for each DepthBiasConstantFactor
-		if (bHasEXTExtendedDynamicState3)
-		{
-			LastDynamicStateUpdate->Rasterizer.DepthBiasConstantFactor = OldDepthBiasConstantFactor;
-		}
+		VulkanRHI::vkCmdSetDepthBias(CommandBuffer.GetHandle(), NewDynamicState.Rasterizer.DepthBiasConstantFactor, depthBiasClamp, NewDynamicState.Rasterizer.DepthBiasSlopeScale);
 		VULKAN_DYNAMIC_STATE_BLOCK_END();
 	}
 
@@ -789,6 +791,22 @@ void FVulkanPendingGfxState::UpdateDynamicStates(FVulkanCommandBuffer& CommandBu
 
 	if (Device.SupportsShaderObjects())
 	{
+		VULKAN_DYNAMIC_STATE_BLOCK_START(DepthStencil.FrontCompareMask);
+		VulkanRHI::vkCmdSetStencilCompareMask(CommandBuffer.GetHandle(), VK_STENCIL_FACE_FRONT_BIT, NewDynamicState.DepthStencil.FrontCompareMask);
+		VULKAN_DYNAMIC_STATE_BLOCK_END();
+
+		VULKAN_DYNAMIC_STATE_BLOCK_START(DepthStencil.BackCompareMask);
+		VulkanRHI::vkCmdSetStencilCompareMask(CommandBuffer.GetHandle(), VK_STENCIL_FACE_BACK_BIT, NewDynamicState.DepthStencil.BackCompareMask);
+		VULKAN_DYNAMIC_STATE_BLOCK_END();
+
+		VULKAN_DYNAMIC_STATE_BLOCK_START(DepthStencil.FrontWriteMask);
+		VulkanRHI::vkCmdSetStencilWriteMask(CommandBuffer.GetHandle(), VK_STENCIL_FACE_FRONT_BIT, NewDynamicState.DepthStencil.FrontWriteMask);
+		VULKAN_DYNAMIC_STATE_BLOCK_END();
+
+		VULKAN_DYNAMIC_STATE_BLOCK_START(DepthStencil.BackWriteMask);
+		VulkanRHI::vkCmdSetStencilWriteMask(CommandBuffer.GetHandle(), VK_STENCIL_FACE_BACK_BIT, NewDynamicState.DepthStencil.BackWriteMask);
+		VULKAN_DYNAMIC_STATE_BLOCK_END();
+
 		if (GRHISupportsPipelineVariableRateShading && (GRHIVariableRateShadingImageDataType == VRSImage_Palette))
 		{
 			VULKAN_DYNAMIC_STATE_BLOCK_START(ShadingRate, RenderTargets.bHasFragmentDensityAttachment, Combiner);
diff --git a/Engine/Source/Runtime/VulkanRHI/Private/VulkanPipeline.cpp b/Engine/Source/Runtime/VulkanRHI/Private/VulkanPipeline.cpp
--- a/Engine/Source/Runtime/VulkanRHI/Private/VulkanPipeline.cpp
+++ b/Engine/Source/Runtime/VulkanRHI/Private/VulkanPipeline.cpp
@@ -1325,6 +1325,7 @@ static TArray<VkDynamicState> FillDynamicStates(const FVulkanDevice& Device, con
 	if (bHasEXTExtendedDynamicState2)
 	{
 		DynamicStates.Add(VK_DYNAMIC_STATE_DEPTH_BIAS_ENABLE);
+		DynamicStates.Add(VK_DYNAMIC_STATE_DEPTH_BIAS);  // from Vulkan 1.0
 
 		// VK_DYNAMIC_STATE_LOGIC_OP_EXT not used
 		// VK_DYNAMIC_STATE_PATCH_CONTROL_POINTS_EXT not used
