VulkanRHI: Fix signed/unsigned mismatch causing uniform buffers to get freed immediately after first use on startup.

#rnx

diff --git a/Engine/Source/Runtime/VulkanRHI/Private/VulkanMemory.h b/Engine/Source/Runtime/VulkanRHI/Private/VulkanMemory.h
--- a/Engine/Source/Runtime/VulkanRHI/Private/VulkanMemory.h
+++ b/Engine/Source/Runtime/VulkanRHI/Private/VulkanMemory.h
@@ -859,7 +859,7 @@ namespace VulkanRHI
 		struct FUBPendingFree
 		{
 			FVulkanAllocation Allocation;
-			uint64 Frame = 0;
+			int64 Frame = 0;
 		};
 
 		struct
