VulkanRHI: Use cvar instead of hardcoded value to delay frees (helpful to debug lifetime issues).  Will reduce additional delay for some resources (UB, Staging) from 3 frames to 2 frames.

#rnx

diff --git a/Engine/Source/Runtime/VulkanRHI/Private/VulkanMemory.cpp b/Engine/Source/Runtime/VulkanRHI/Private/VulkanMemory.cpp
--- a/Engine/Source/Runtime/VulkanRHI/Private/VulkanMemory.cpp
+++ b/Engine/Source/Runtime/VulkanRHI/Private/VulkanMemory.cpp
@@ -28,7 +28,7 @@ THIRD_PARTY_INCLUDES_END
 uint32 GVulkanRHIDeletionFrameNumber = 0;
 
 
-static int32 GVulkanNumFramesToWaitForResourceDelete = 2;
+int32 GVulkanNumFramesToWaitForResourceDelete = 2;
 static FAutoConsoleVariableRef CVarVulkanNumFramesToWaitForResourceDelete(
 	TEXT("r.Vulkan.NumFramesToWaitForResourceDelete"),
 	GVulkanNumFramesToWaitForResourceDelete,
@@ -2245,7 +2245,7 @@ namespace VulkanRHI
 				for (int32 Index = 0; Index < FreeAllocations.Num(); ++Index)
 				{
 					FVulkanSubresourceAllocator* BufferAllocation = FreeAllocations[Index];
-					if (bImmediately || BufferAllocation->FrameFreed + NUM_FRAMES_TO_WAIT_BEFORE_RELEASING_TO_OS < GFrameNumberRenderThread)
+					if (bImmediately || BufferAllocation->FrameFreed + GVulkanNumFramesToWaitForResourceDelete < GFrameNumberRenderThread)
 					{
 						BufferAllocationsToRelease.Add(BufferAllocation);
 						FreeAllocations.RemoveAtSwap(Index, EAllowShrinking::No);
@@ -3703,7 +3703,7 @@ namespace VulkanRHI
 			}
 
 			// making use of the fact that we always add to the end of the array, so allocations are sorted by frame ascending
-			int32 OldestFrameToKeep = GFrameNumberRenderThread - VulkanRHI::NUM_FRAMES_TO_WAIT_BEFORE_RELEASING_TO_OS;
+			int32 OldestFrameToKeep = GFrameNumberRenderThread - GVulkanNumFramesToWaitForResourceDelete;
 			int32 NumAlloc = UBAllocations.PendingFree.Num();
 			int32 Index = 0;
 			for (; Index < NumAlloc; ++Index)
@@ -3716,7 +3716,7 @@ namespace VulkanRHI
 				else
 				{
 					// calculate when we will be able to delete the oldest allocation
-					GFrameNumberRenderThread_WhenWeCanDelete = Alloc.Frame + VulkanRHI::NUM_FRAMES_TO_WAIT_BEFORE_RELEASING_TO_OS + 1;
+					GFrameNumberRenderThread_WhenWeCanDelete = Alloc.Frame + GVulkanNumFramesToWaitForResourceDelete + 1;
 					break;
 				}
 			}
@@ -4503,7 +4503,7 @@ namespace VulkanRHI
 			for (int32 Index = NumFreeBuffers - 1; Index >= 0; --Index)
 			{
 				FFreeEntry& Entry = FreeStagingBuffers[Index];
-				if (bImmediately || Entry.FrameNumber + NUM_FRAMES_TO_WAIT_BEFORE_RELEASING_TO_OS < GFrameNumberRenderThread)
+				if (bImmediately || Entry.FrameNumber + GVulkanNumFramesToWaitForResourceDelete < GFrameNumberRenderThread)
 				{
 					UsedMemory -= Entry.StagingBuffer->GetSize();
 					Entry.StagingBuffer->Destroy();
diff --git a/Engine/Source/Runtime/VulkanRHI/Private/VulkanMemory.h b/Engine/Source/Runtime/VulkanRHI/Private/VulkanMemory.h
--- a/Engine/Source/Runtime/VulkanRHI/Private/VulkanMemory.h
+++ b/Engine/Source/Runtime/VulkanRHI/Private/VulkanMemory.h
@@ -130,11 +130,6 @@ extern int32 GVulkanUseBufferBinning;
 
 namespace VulkanRHI
 {
-	enum
-	{
-		NUM_FRAMES_TO_WAIT_BEFORE_RELEASING_TO_OS = 3,
-	};
-
 	enum EVulkanAllocationType : uint8 {
 		EVulkanAllocationEmpty,
 		EVulkanAllocationPooledBuffer,
diff --git a/Engine/Source/Runtime/VulkanRHI/Private/VulkanPendingState.cpp b/Engine/Source/Runtime/VulkanRHI/Private/VulkanPendingState.cpp
--- a/Engine/Source/Runtime/VulkanRHI/Private/VulkanPendingState.cpp
+++ b/Engine/Source/Runtime/VulkanRHI/Private/VulkanPendingState.cpp
@@ -11,6 +11,8 @@
 #include "VulkanPipeline.h"
 #include "VulkanContext.h"
 
+extern int32 GVulkanNumFramesToWaitForResourceDelete;
+
 FVulkanDescriptorPool::FVulkanDescriptorPool(FVulkanDevice* InDevice, const FVulkanDescriptorSetsLayout& InLayout, uint32 MaxSetsAllocations)
 	: Device(InDevice)
 	, MaxDescriptorSets(0)
@@ -292,7 +294,7 @@ void FVulkanDescriptorPoolsManager::GC()
 	for (int32 Index = PoolSets.Num() - 1; Index >= 0; Index--)
 	{
 		auto* PoolSet = PoolSets[Index];
-		if (PoolSet->IsUnused() && GFrameNumberRenderThread - PoolSet->GetLastFrameUsed() > VulkanRHI::NUM_FRAMES_TO_WAIT_BEFORE_RELEASING_TO_OS)
+		if (PoolSet->IsUnused() && ((int32)(GFrameNumberRenderThread - PoolSet->GetLastFrameUsed()) > GVulkanNumFramesToWaitForResourceDelete))
 		{
 			PoolSets.RemoveAtSwap(Index, EAllowShrinking::Yes);
 
@@ -1168,8 +1170,7 @@ void FVulkanDescriptorSetCache::FCachedPool::Reset()
 
 bool FVulkanDescriptorSetCache::FCachedPool::CanGC() const
 {
-	constexpr uint32 FramesBeforeGC = VulkanRHI::NUM_FRAMES_TO_WAIT_BEFORE_RELEASING_TO_OS;
-	return ((GFrameNumberRenderThread - RecentFrame) > FramesBeforeGC);
+	return ((int32)(GFrameNumberRenderThread - RecentFrame) > GVulkanNumFramesToWaitForResourceDelete);
 }
 
 float FVulkanDescriptorSetCache::FCachedPool::CalcAllocRatio() const
