VulkanRHI: Use ref instead of pointer (fix issue exposed when using debug mode that waits for idle on every submit).

#rnx

diff --git a/Engine/Source/Runtime/VulkanRHI/Private/VulkanSubmission.cpp b/Engine/Source/Runtime/VulkanRHI/Private/VulkanSubmission.cpp
--- a/Engine/Source/Runtime/VulkanRHI/Private/VulkanSubmission.cpp
+++ b/Engine/Source/Runtime/VulkanRHI/Private/VulkanSubmission.cpp
@@ -614,7 +614,7 @@ void FVulkanDynamicRHI::EnqueueEndOfPipeTask(TUniqueFunction<void()> TaskFunc, T
 
 void FVulkanDynamicRHI::RHIBlockUntilGPUIdle()
 {
-	TArray<FVulkanSyncPoint*, TInlineAllocator<(int32)EVulkanQueueType::Count>> EndOfQueueEvent;
+	TArray<FVulkanSyncPointRef, TInlineAllocator<(int32)EVulkanQueueType::Count>> EndOfQueueEvent;
 
 	// Create payloads with a signal for each queue
 	{
@@ -637,7 +637,7 @@ void FVulkanDynamicRHI::RHIBlockUntilGPUIdle()
 	KickSubmissionThread();
 
 	// Wait on each event
-	for (FVulkanSyncPoint* Event : EndOfQueueEvent)
+	for (FVulkanSyncPointRef& Event : EndOfQueueEvent)
 	{
 		ProcessInterruptQueueUntil(Event);
 	}
