VulkanRHI: Add push descriptor functions.  Add core 1.4 feature query struct.

#rnx

diff --git a/Engine/Source/Runtime/VulkanRHI/Private/VulkanCommandWrappers.h b/Engine/Source/Runtime/VulkanRHI/Private/VulkanCommandWrappers.h
--- a/Engine/Source/Runtime/VulkanRHI/Private/VulkanCommandWrappers.h
+++ b/Engine/Source/Runtime/VulkanRHI/Private/VulkanCommandWrappers.h
@@ -186,6 +186,8 @@ struct FWrapLayer
 	static void CmdWriteTimestamp(VkResult Result, VkCommandBuffer CommandBuffer, VkPipelineStageFlagBits PipelineStage, VkQueryPool QueryPool, uint32 Query) VULKAN_LAYER_BODY
 	static void CmdCopyQueryPoolResults(VkResult Result, VkCommandBuffer CommandBuffer, VkQueryPool QueryPool, uint32 FirstQuery, uint32 QueryCount, VkBuffer DstBuffer, VkDeviceSize DstOffset, VkDeviceSize Stride, VkQueryResultFlags Flags) VULKAN_LAYER_BODY
 	static void CmdPushConstants(VkResult Result, VkCommandBuffer CommandBuffer, VkPipelineLayout Layout, VkShaderStageFlags StageFlags, uint32_t Offset, uint32_t Size, const void* pValues) VULKAN_LAYER_BODY
+	static void CmdPushDescriptorSetKHR(VkResult Result, VkCommandBuffer CommandBuffer, VkPipelineBindPoint PipelineBindPoint, VkPipelineLayout Layout, uint32_t Set, uint32_t DescriptorWriteCount, const VkWriteDescriptorSet* DescriptorWrites) VULKAN_LAYER_BODY
+	static void CmdPushDescriptorSet2KHR(VkResult Result, VkCommandBuffer CommandBuffer, const VkPushDescriptorSetInfo* PushDescriptorSetInfo) VULKAN_LAYER_BODY
 	static void CmdBeginRenderPass(VkResult Result, VkCommandBuffer CommandBuffer, const VkRenderPassBeginInfo* RenderPassBegin, VkSubpassContents Contents) VULKAN_LAYER_BODY
 	static void CmdBeginRenderPass2KHR(VkResult Result, VkCommandBuffer CommandBuffer, const VkRenderPassBeginInfo* RenderPassBegin, const VkSubpassBeginInfo* pSubpassBeginInfo) VULKAN_LAYER_BODY
 	static void CmdNextSubpass(VkResult Result, VkCommandBuffer CommandBuffer, VkSubpassContents Contents) VULKAN_LAYER_BODY
@@ -194,7 +196,6 @@ struct FWrapLayer
 	static void CmdEndRenderingKHR(VkResult Result, VkCommandBuffer CommandBuffer) VULKAN_LAYER_BODY
 	static void CmdExecuteCommands(VkResult Result, VkCommandBuffer CommandBuffer, uint32 CommandBufferCount, const VkCommandBuffer* pCommandBuffers) VULKAN_LAYER_BODY
 
-
 	static void QueuePresent(VkResult Result, VkQueue Queue, const VkPresentInfoKHR* PresentInfo) VULKAN_LAYER_BODY
 	static void GetSwapChainImagesKHR(VkResult Result, VkDevice Device, VkSwapchainKHR Swapchain, uint32_t* SwapchainImageCount, VkImage* SwapchainImages) VULKAN_LAYER_BODY
 	static void CreateSwapchainKHR(VkResult Result, VkDevice Device, const VkSwapchainCreateInfoKHR* CreateInfo, VkSwapchainKHR* Swapchain) VULKAN_LAYER_BODY
@@ -1394,6 +1395,20 @@ namespace VulkanRHI
 		FWrapLayer::CmdPushConstants(VK_SUCCESS, CommandBuffer, Layout, StageFlags, Offset, Size, pValues);
 	}
 
+	static FORCEINLINE_DEBUGGABLE void  vkCmdPushDescriptorSetKHR(VkCommandBuffer CommandBuffer, VkPipelineBindPoint PipelineBindPoint, VkPipelineLayout Layout, uint32_t Set, uint32_t DescriptorWriteCount, const VkWriteDescriptorSet* DescriptorWrites)
+	{
+		FWrapLayer::CmdPushDescriptorSetKHR(VK_RESULT_MAX_ENUM, CommandBuffer, PipelineBindPoint, Layout, Set, DescriptorWriteCount, DescriptorWrites);
+		VULKANAPINAMESPACE::vkCmdPushDescriptorSetKHR(CommandBuffer, PipelineBindPoint, Layout, Set, DescriptorWriteCount, DescriptorWrites);
+		FWrapLayer::CmdPushDescriptorSetKHR(VK_SUCCESS, CommandBuffer, PipelineBindPoint, Layout, Set, DescriptorWriteCount, DescriptorWrites);
+	}
+
+	static FORCEINLINE_DEBUGGABLE void  vkCmdPushDescriptorSet2KHR(VkCommandBuffer CommandBuffer, const VkPushDescriptorSetInfo* PushDescriptorSetInfo)
+	{
+		FWrapLayer::CmdPushDescriptorSet2KHR(VK_RESULT_MAX_ENUM, CommandBuffer, PushDescriptorSetInfo);
+		VULKANAPINAMESPACE::vkCmdPushDescriptorSet2KHR(CommandBuffer, PushDescriptorSetInfo);
+		FWrapLayer::CmdPushDescriptorSet2KHR(VK_SUCCESS, CommandBuffer, PushDescriptorSetInfo);
+	}
+
 	static FORCEINLINE_DEBUGGABLE void  vkCmdBeginRenderPass(VkCommandBuffer CommandBuffer, const VkRenderPassBeginInfo* RenderPassBegin, VkSubpassContents Contents)
 	{
 		FWrapLayer::CmdBeginRenderPass(VK_RESULT_MAX_ENUM, CommandBuffer, RenderPassBegin, Contents);
diff --git a/Engine/Source/Runtime/VulkanRHI/Private/VulkanDebug.cpp b/Engine/Source/Runtime/VulkanRHI/Private/VulkanDebug.cpp
--- a/Engine/Source/Runtime/VulkanRHI/Private/VulkanDebug.cpp
+++ b/Engine/Source/Runtime/VulkanRHI/Private/VulkanDebug.cpp
@@ -2413,6 +2413,26 @@ void FWrapLayer::CmdPushConstants(VkResult Result, VkCommandBuffer CommandBuffer
 	}
 }
 
+void FWrapLayer::CmdPushDescriptorSetKHR(VkResult Result, VkCommandBuffer CommandBuffer, VkPipelineBindPoint PipelineBindPoint, VkPipelineLayout Layout, uint32_t Set, uint32_t DescriptorWriteCount, const VkWriteDescriptorSet* DescriptorWrites)
+{
+	if (Result == VK_RESULT_MAX_ENUM)
+	{
+#if VULKAN_ENABLE_DUMP_LAYER
+		CmdPrintfBegin(CommandBuffer, FString::Printf(TEXT("vkCmdPushDescriptorSetKHR()")));
+#endif
+	}
+}
+
+void FWrapLayer::CmdPushDescriptorSet2KHR(VkResult Result, VkCommandBuffer CommandBuffer, const VkPushDescriptorSetInfo* PushDescriptorSetInfo)
+{
+	if (Result == VK_RESULT_MAX_ENUM)
+	{
+#if VULKAN_ENABLE_DUMP_LAYER
+		CmdPrintfBegin(CommandBuffer, FString::Printf(TEXT("vkCmdPushDescriptorSet2KHR()")));
+#endif
+	}
+}
+
 void FWrapLayer::CmdBeginRenderPass(VkResult Result, VkCommandBuffer CommandBuffer, const VkRenderPassBeginInfo* RenderPassBegin, VkSubpassContents Contents)
 {
 	if (Result == VK_RESULT_MAX_ENUM)
diff --git a/Engine/Source/Runtime/VulkanRHI/Private/VulkanDevice.cpp b/Engine/Source/Runtime/VulkanRHI/Private/VulkanDevice.cpp
--- a/Engine/Source/Runtime/VulkanRHI/Private/VulkanDevice.cpp
+++ b/Engine/Source/Runtime/VulkanRHI/Private/VulkanDevice.cpp
@@ -221,6 +221,12 @@ void FVulkanPhysicalDeviceFeatures::Query(VkPhysicalDevice PhysicalDevice, uint3
 		Core_1_3.sType = VK_STRUCTURE_TYPE_PHYSICAL_DEVICE_VULKAN_1_3_FEATURES;
 	}
 
+	if (APIVersion >= VK_API_VERSION_1_4)
+	{
+		Core_1_3.pNext = &Core_1_4;
+		Core_1_4.sType = VK_STRUCTURE_TYPE_PHYSICAL_DEVICE_VULKAN_1_4_FEATURES;
+	}
+
 	VulkanRHI::vkGetPhysicalDeviceFeatures2(PhysicalDevice, &PhysicalDeviceFeatures2);
 
 	// Copy features into old struct for convenience
diff --git a/Engine/Source/Runtime/VulkanRHI/Private/VulkanDevice.h b/Engine/Source/Runtime/VulkanRHI/Private/VulkanDevice.h
--- a/Engine/Source/Runtime/VulkanRHI/Private/VulkanDevice.h
+++ b/Engine/Source/Runtime/VulkanRHI/Private/VulkanDevice.h
@@ -111,6 +111,7 @@ struct FOptionalVulkanDeviceExtensions
 			// Promoted to 1.4
 			uint64 HasKHRMaintenance5 : 1;
 			uint64 HasKHRMaintenance6 : 1;
+			uint64 HasKHRPushDescriptors : 1;
 			uint64 HasEXTLoadStoreOpNone : 1;
 			uint64 HasEXTHostImageCopy : 1;
 		};
@@ -193,6 +194,7 @@ private:
 	// Anything above Core 1.1 cannot be assumed, they should only be used by the device at init time
 	VkPhysicalDeviceVulkan12Features Core_1_2;
 	VkPhysicalDeviceVulkan13Features Core_1_3;
+	VkPhysicalDeviceVulkan14Features Core_1_4;
 
 	friend class FVulkanDevice;
 	friend FVulkanDynamicRHI;
diff --git a/Engine/Source/Runtime/VulkanRHI/Private/VulkanLoader.h b/Engine/Source/Runtime/VulkanRHI/Private/VulkanLoader.h
--- a/Engine/Source/Runtime/VulkanRHI/Private/VulkanLoader.h
+++ b/Engine/Source/Runtime/VulkanRHI/Private/VulkanLoader.h
@@ -146,6 +146,8 @@
 	EnumMacro(PFN_vkCmdWriteTimestamp, vkCmdWriteTimestamp) \
 	EnumMacro(PFN_vkCmdCopyQueryPoolResults, vkCmdCopyQueryPoolResults) \
 	EnumMacro(PFN_vkCmdPushConstants, vkCmdPushConstants) \
+	EnumMacro(PFN_vkCmdPushDescriptorSetKHR, vkCmdPushDescriptorSetKHR) \
+	EnumMacro(PFN_vkCmdPushDescriptorSet2KHR, vkCmdPushDescriptorSet2KHR) \
 	EnumMacro(PFN_vkCmdBeginRenderPass, vkCmdBeginRenderPass) \
 	EnumMacro(PFN_vkCmdNextSubpass, vkCmdNextSubpass) \
 	EnumMacro(PFN_vkCmdEndRenderPass, vkCmdEndRenderPass) \
