VulkanRHI: Add message when timeline semaphore value query returns -1.

#rnx

diff --git a/Engine/Source/Runtime/VulkanRHI/Private/VulkanSynchronization.cpp b/Engine/Source/Runtime/VulkanRHI/Private/VulkanSynchronization.cpp
--- a/Engine/Source/Runtime/VulkanRHI/Private/VulkanSynchronization.cpp
+++ b/Engine/Source/Runtime/VulkanRHI/Private/VulkanSynchronization.cpp
@@ -254,5 +254,9 @@ uint64 FVulkanSemaphore::GetTimelineSemaphoreValue()
 	checkf(IsTimeline(), TEXT("GetTimelineSemaphoreValue not available!  Semaphore was not created with EVulkanSemaphoreFlags::Timeline."));
 	uint64 Value = 0;
 	VERIFYVULKANRESULT(VulkanRHI::vkGetSemaphoreCounterValueKHR(Device.GetHandle(), SemaphoreHandle, (uint64_t*)&Value));
+	if (Value == MAX_uint64)
+	{
+		UE_LOG(LogVulkanRHI, Warning, TEXT("Timeline semaphore value is MAX_uint64 (this can happen when the device is lost)."));
+	}
 	return Value;
 }
