VulkanRHI: Move push descriptor functions to optional

#rnx

diff --git a/Engine/Source/Runtime/VulkanRHI/Private/VulkanLoader.h b/Engine/Source/Runtime/VulkanRHI/Private/VulkanLoader.h
--- a/Engine/Source/Runtime/VulkanRHI/Private/VulkanLoader.h
+++ b/Engine/Source/Runtime/VulkanRHI/Private/VulkanLoader.h
@@ -146,8 +146,6 @@
 	EnumMacro(PFN_vkCmdWriteTimestamp, vkCmdWriteTimestamp) \
 	EnumMacro(PFN_vkCmdCopyQueryPoolResults, vkCmdCopyQueryPoolResults) \
 	EnumMacro(PFN_vkCmdPushConstants, vkCmdPushConstants) \
-	EnumMacro(PFN_vkCmdPushDescriptorSetKHR, vkCmdPushDescriptorSetKHR) \
-	EnumMacro(PFN_vkCmdPushDescriptorSet2KHR, vkCmdPushDescriptorSet2KHR) \
 	EnumMacro(PFN_vkCmdBeginRenderPass, vkCmdBeginRenderPass) \
 	EnumMacro(PFN_vkCmdNextSubpass, vkCmdNextSubpass) \
 	EnumMacro(PFN_vkCmdEndRenderPass, vkCmdEndRenderPass) \
@@ -277,7 +275,9 @@
 	EnumMacro(PFN_vkCmdSetSampleMaskEXT, vkCmdSetSampleMaskEXT) \
 	EnumMacro(PFN_vkCmdSetFragmentShadingRateKHR, vkCmdSetFragmentShadingRateKHR) \
 	EnumMacro(PFN_vkCopyMemoryToImageEXT, vkCopyMemoryToImageEXT) \
-	EnumMacro(PFN_vkTransitionImageLayoutEXT, vkTransitionImageLayoutEXT)
+	EnumMacro(PFN_vkTransitionImageLayoutEXT, vkTransitionImageLayoutEXT) \
+	EnumMacro(PFN_vkCmdPushDescriptorSetKHR, vkCmdPushDescriptorSetKHR) \
+	EnumMacro(PFN_vkCmdPushDescriptorSet2KHR, vkCmdPushDescriptorSet2KHR)
 
 
 
