add -clearPSODriverCache support to OpenGL and Vulkan RHIs.

diff --git a/Engine/Source/Runtime/OpenGLDrv/Private/OpenGLProgramBinaryFileCache.cpp b/Engine/Source/Runtime/OpenGLDrv/Private/OpenGLProgramBinaryFileCache.cpp
--- a/Engine/Source/Runtime/OpenGLDrv/Private/OpenGLProgramBinaryFileCache.cpp
+++ b/Engine/Source/Runtime/OpenGLDrv/Private/OpenGLProgramBinaryFileCache.cpp
@@ -355,8 +355,8 @@ void FOpenGLProgramBinaryCache::Initialize()
 	CacheFolderPathRoot = FPaths::ProjectSavedDir() / TEXT("RHICache") / TEXT("ProgramBinaryCache");
 #endif
 
-	// Remove entire ProgramBinaryCache folder if -ClearOpenGLBinaryProgramCache is specified on command line
-	if (FParse::Param(FCommandLine::Get(), TEXT("ClearOpenGLBinaryProgramCache")))
+	// Remove entire ProgramBinaryCache folder if either -ClearOpenGLBinaryProgramCache or -clearPSODriverCache are specified on command line
+	if (FParse::Param(FCommandLine::Get(), TEXT("ClearOpenGLBinaryProgramCache")) || FParse::Param(FCommandLine::Get(), TEXT("clearPSODriverCache")))
 	{
 		UE_LOG(LogRHI, Log, TEXT("ClearOpenGLBinaryProgramCache specified, deleting binary program cache folder: %s"), *CacheFolderPathRoot);
 		FPlatformFileManager::Get().GetPlatformFile().DeleteDirectoryRecursively(*CacheFolderPathRoot);
diff --git a/Engine/Source/Runtime/VulkanRHI/Private/VulkanChunkedPipelineCache.cpp b/Engine/Source/Runtime/VulkanRHI/Private/VulkanChunkedPipelineCache.cpp
--- a/Engine/Source/Runtime/VulkanRHI/Private/VulkanChunkedPipelineCache.cpp
+++ b/Engine/Source/Runtime/VulkanRHI/Private/VulkanChunkedPipelineCache.cpp
@@ -465,6 +465,7 @@ public:
 
 		VkPipelineCacheCreateInfo PipelineCacheInfo;
 		ZeroVulkanStruct(PipelineCacheInfo, VK_STRUCTURE_TYPE_PIPELINE_CACHE_CREATE_INFO);
+		check(PipelineCacheObj == VK_NULL_HANDLE);
 		VERIFYVULKANRESULT(VulkanRHI::vkCreatePipelineCache(UE::Vulkan::GetVulkanDeviceHandle(), &PipelineCacheInfo, VULKAN_CPU_ALLOCATOR, &PipelineCacheObj));
 		CacheState = ECacheState::Building;
 		check(LastSaveSize == 0);
@@ -748,6 +749,7 @@ public:
 			ZeroVulkanStruct(PipelineCacheInfo, VK_STRUCTURE_TYPE_PIPELINE_CACHE_CREATE_INFO);
 			PipelineCacheInfo.pInitialData = CacheBytes.GetData();
 			PipelineCacheInfo.initialDataSize = CacheBytes.Num();
+			check(PipelineCacheObj == VK_NULL_HANDLE);
 			VERIFYVULKANRESULT(VulkanRHI::vkCreatePipelineCache(UE::Vulkan::GetVulkanDeviceHandle(), &PipelineCacheInfo, VULKAN_CPU_ALLOCATOR, &PipelineCacheObj));
 		}
 
@@ -995,6 +997,8 @@ private:
 		const FString CacheSubDir = UE::Vulkan::GetPSOBinaryCacheVersionedFolder();
 		const FString CombinedCacheSubDir = FPaths::Combine(RootCacheFolder, CacheSubDir);
 
+		const bool bDeleteAll = FParse::Param(FCommandLine::Get(), TEXT("ClearVulkanBinaryProgramCache")) || FParse::Param(FCommandLine::Get(), TEXT("clearPSODriverCache"));
+
 		// delete anything unexpected from the RHI cache root folder.
 		IPlatformFile& PlatformFile = FPlatformFileManager::Get().GetPlatformFile();
 		TArray<FString> FoundFiles;
@@ -1003,7 +1007,7 @@ private:
 		{
 			const FString FullPath = (RootCacheFolder / FoundFile);
 			const bool bIsDir = PlatformFile.DirectoryExists(*FullPath);
-			if (FoundFile != CacheSubDir || !bIsDir)
+			if (FoundFile != CacheSubDir || !bIsDir || bDeleteAll)
 			{
 				bool bSuccess;
 				if (bIsDir)
diff --git a/Engine/Source/Runtime/VulkanRHI/Private/VulkanPipeline.cpp b/Engine/Source/Runtime/VulkanRHI/Private/VulkanPipeline.cpp
--- a/Engine/Source/Runtime/VulkanRHI/Private/VulkanPipeline.cpp
+++ b/Engine/Source/Runtime/VulkanRHI/Private/VulkanPipeline.cpp
@@ -493,8 +493,8 @@ void FVulkanPipelineStateCacheManager::InitAndLoad(const TArray<FString>& CacheF
 
 			CompiledPSOCacheTopFolderPath = FVulkanPlatform::GetCompiledPSOCacheTopFolderPath();
 
-			// Remove entire ProgramBinaryCache folder if -ClearOpenGLBinaryProgramCache is specified on command line
-			if (FParse::Param(FCommandLine::Get(), TEXT("ClearVulkanBinaryProgramCache")))
+			// Remove entire ProgramBinaryCache folder if either -ClearVulkanBinaryProgramCache or clearPSODriverCache are specified on command line
+			if (FParse::Param(FCommandLine::Get(), TEXT("ClearVulkanBinaryProgramCache")) || FParse::Param(FCommandLine::Get(), TEXT("clearPSODriverCache")))
 			{
 				UE_LOG(LogVulkanRHI, Log, TEXT("FVulkanPipelineStateCacheManager: Deleting binary program cache folder for -ClearVulkanBinaryProgramCache: %s"), *CompiledPSOCacheTopFolderPath);
 				FPlatformFileManager::Get().GetPlatformFile().DeleteDirectoryRecursively(*CompiledPSOCacheTopFolderPath);
