PR #13935: fix build error for android platform when enable vulkan dump layer

#rnx

diff --git a/Engine/Source/Runtime/VulkanRHI/Private/VulkanDebug.cpp b/Engine/Source/Runtime/VulkanRHI/Private/VulkanDebug.cpp
--- a/Engine/Source/Runtime/VulkanRHI/Private/VulkanDebug.cpp
+++ b/Engine/Source/Runtime/VulkanRHI/Private/VulkanDebug.cpp
@@ -4281,7 +4281,7 @@ void FWrapLayer::CmdTraceRaysIndirectKHR(VkResult Result, VkCommandBuffer Comman
 	}
 }
 
-void FWrapLayer::CmdTraceRaysIndirect2KHR(VkResult Result, VkDeviceAddress IndirectDeviceAddress)
+void FWrapLayer::CmdTraceRaysIndirect2KHR(VkResult Result, VkCommandBuffer CommandBuffer, VkDeviceAddress IndirectDeviceAddress)
 {
 	if (Result == VK_RESULT_MAX_ENUM)
 	{
