VulkanRHI: Use push descriptors for UB bindings when in bindless

#rnx

diff --git a/Engine/Shaders/Public/Platform/Vulkan/VulkanCommon.ush b/Engine/Shaders/Public/Platform/Vulkan/VulkanCommon.ush
--- a/Engine/Shaders/Public/Platform/Vulkan/VulkanCommon.ush
+++ b/Engine/Shaders/Public/Platform/Vulkan/VulkanCommon.ush
@@ -7,7 +7,7 @@
 #pragma once
 
 // Update this GUID to invalidate shader recompilation for only Vulkan shaders
-#pragma message("UESHADERMETADATA_VERSION 30B5352B-3DD3-4A44-A103-5D1B9086A805")
+#pragma message("UESHADERMETADATA_VERSION EBB4295A-3A82-45E7-89AD-5816B49A55DD")
 
 #define COMPILER_SUPPORTS_ULONG_TYPES 1
 #define COMPILER_SUPPORTS_UINT64_IMAGE_ATOMICS (1 && COMPILER_SUPPORTS_ULONG_TYPES)
diff --git a/Engine/Source/Developer/VulkanShaderFormat/Private/SpirVShaderCompiler.inl b/Engine/Source/Developer/VulkanShaderFormat/Private/SpirVShaderCompiler.inl
--- a/Engine/Source/Developer/VulkanShaderFormat/Private/SpirVShaderCompiler.inl
+++ b/Engine/Source/Developer/VulkanShaderFormat/Private/SpirVShaderCompiler.inl
@@ -616,6 +616,8 @@ static void GatherSpirvReflectionBindings(
 					Reflection.ChangeDescriptorBindingNumbers(pBinding, BindingOffset + pBinding->binding, VulkanBindless::BindlessSingleUseUniformBufferSet);
 				}
 			}
+
+			checkf(OutBindings.UniformBuffers.Num() < VulkanBindless::MaxUniformBuffersPerStage, TEXT("Too many UB bindings in shader."));
 		}
 	}
 }
@@ -703,7 +705,7 @@ static bool BuildShaderOutputFromSpirv(
 	// Final descriptor binding numbers for all other resource types
 	{
 		const ShaderStage::EStage UEStage = ShaderStage::GetStageForFrequency(InternalState.GetShaderFrequency());
-		const int32 StageOffset = InternalState.bSupportsBindless ? (UEStage * VulkanBindless::MaxUniformBuffersPerStage) : 0;
+		const int32 StageOffset = InternalState.bSupportsBindless ? VulkanBindless::GetOffsetForFrequency(InternalState.GetShaderFrequency()) : 0;
 		const uint32_t DescSetNumber = InternalState.bSupportsBindless ? (uint32_t)VulkanBindless::BindlessSingleUseUniformBufferSet : (uint32_t)UEStage;
 
 		auto AddShaderValidationType = [](uint32_t VulkanBindingIndex, const FShaderParameterParser::FParsedShaderParameter* ParsedParam, FShaderCompilerOutput& Output) 
@@ -737,7 +739,7 @@ static bool BuildShaderOutputFromSpirv(
 				}
 
 				const int32 BindingSlot = SerializedOutput.Header.Bindings.Num();
-				const int32 BindingIndex = StageOffset + BindingSlot;
+				const int32 BindingIndex = (StageOffset < 0) ? VulkanBindless::MaxUniformBuffersTotal - BindingSlot - 1 : StageOffset + BindingSlot;
 				FVulkanShaderHeader::FBindingInfo& BindingInfo = SerializedOutput.Header.Bindings.AddZeroed_GetRef();
 				BindingInfo.DescriptorType = DescriptorType;
 
@@ -895,6 +897,7 @@ static bool BuildShaderOutputFromSpirv(
 		const int32 UBOBindings = AddReflectionInfos(Bindings.UniformBuffers, VK_DESCRIPTOR_TYPE_UNIFORM_BUFFER, GlobalUBCount);
 		SerializedOutput.Header.NumBoundUniformBuffers = UBOBindings;
 		SerializedOutput.PackedResourceCounts.NumCBs = (uint8)UBOBindings;
+		checkf(UBOBindings < VulkanBindless::MaxUniformBuffersPerStage, TEXT("Too many UB bindings in shader."));
 
 		AddReflectionInfos(Bindings.InputAttachments, VK_DESCRIPTOR_TYPE_INPUT_ATTACHMENT, 0);
 
diff --git a/Engine/Source/Runtime/VulkanRHI/Private/VulkanBindlessDescriptorManager.cpp b/Engine/Source/Runtime/VulkanRHI/Private/VulkanBindlessDescriptorManager.cpp
--- a/Engine/Source/Runtime/VulkanRHI/Private/VulkanBindlessDescriptorManager.cpp
+++ b/Engine/Source/Runtime/VulkanRHI/Private/VulkanBindlessDescriptorManager.cpp
@@ -69,15 +69,6 @@ static FAutoConsoleVariableRef CVarVulkanBindlessMaxAccelerationStructureCount(
 	ECVF_ReadOnly
 );
 
-int32 GVulkanBindlessBlockSize = 1024 * 1024;
-static FAutoConsoleVariableRef CVarVulkanBindlessBlockSize(
-	TEXT("r.Vulkan.Bindless.BlockSize"),
-	GVulkanBindlessBlockSize,
-	TEXT("Block size to use for single use ub. (default: 1MB)"),
-	ECVF_RenderThreadSafe
-);
-
-
 DECLARE_STATS_GROUP(TEXT("Vulkan Bindless"), STATGROUP_VulkanBindless, STATCAT_Advanced);
 
 DECLARE_DWORD_ACCUMULATOR_STAT(TEXT("Descriptors Allocated"), STAT_VulkanBindlessDescriptorsAllocated, STATGROUP_VulkanBindless);
@@ -361,9 +352,6 @@ void FVulkanBindlessDescriptorManager::Deinit()
 
 		VulkanRHI::vkDestroyDescriptorSetLayout(DeviceHandle, EmptyDescriptorSetLayout, VULKAN_CPU_ALLOCATOR);
 		EmptyDescriptorSetLayout = VK_NULL_HANDLE;
-
-		delete SingleUseUBAllocator;
-		SingleUseUBAllocator = nullptr;
 	}
 }
 
@@ -378,7 +366,6 @@ void FVulkanBindlessDescriptorManager::Init()
 	const VkPhysicalDeviceDescriptorBufferPropertiesEXT& DescriptorBufferProperties = Device.GetOptionalExtensionProperties().DescriptorBufferProps;
 
 	const VkBufferUsageFlags BufferUsageFlags = VK_BUFFER_USAGE_SHADER_DEVICE_ADDRESS_BIT | VK_BUFFER_USAGE_RESOURCE_DESCRIPTOR_BUFFER_BIT_EXT;
-	SingleUseUBAllocator = new VulkanRHI::FTempBlockAllocator(Device, GVulkanBindlessBlockSize, DescriptorBufferProperties.descriptorBufferOffsetAlignment, BufferUsageFlags);
 
 	// Create the dummy layout for unsupported descriptor types
 	{
@@ -478,6 +465,11 @@ void FVulkanBindlessDescriptorManager::Init()
 
 				VkDescriptorSetLayout DescriptorSetLayout = VK_NULL_HANDLE;
 				VERIFYVULKANRESULT(VulkanRHI::vkCreateDescriptorSetLayout(DeviceHandle, &DescriptorSetLayoutCreateInfo, VULKAN_CPU_ALLOCATOR, &DescriptorSetLayout));
+
+				VkDeviceSize BindingOffset = 0;
+				VulkanRHI::vkGetDescriptorSetLayoutBindingOffsetEXT(DeviceHandle, DescriptorSetLayout, 0u, &BindingOffset);
+				check(BindingOffset == 0);
+
 				return DescriptorSetLayout;
 			}
 		};
@@ -591,9 +583,9 @@ void FVulkanBindlessDescriptorManager::Init()
 
 		// Fill in the state for single-use UB
 		// Uniform buffer descriptor set layout differ from the other resources, we reserve a fixed number of descriptors per stage for each draw/dispatch
-		// todo-jn: this could be compacted..
 		{
-			const uint32 NumTotalBindings = VulkanBindless::MaxUniformBuffersPerStage * ShaderStage::MaxNumStages;
+			// Guaranteed to be min 32 by Vulkan 1.4
+			constexpr uint32 NumTotalBindings = VulkanBindless::MaxUniformBuffersTotal;
 
 			TArray<VkDescriptorSetLayoutBinding> DescriptorSetLayoutBindings;
 			DescriptorSetLayoutBindings.SetNumZeroed(NumTotalBindings);
@@ -609,7 +601,7 @@ void FVulkanBindlessDescriptorManager::Init()
 			ZeroVulkanStruct(DescriptorSetLayoutCreateInfo, VK_STRUCTURE_TYPE_DESCRIPTOR_SET_LAYOUT_CREATE_INFO);
 			DescriptorSetLayoutCreateInfo.pBindings = DescriptorSetLayoutBindings.GetData();
 			DescriptorSetLayoutCreateInfo.bindingCount = NumTotalBindings;
-			DescriptorSetLayoutCreateInfo.flags = VK_DESCRIPTOR_SET_LAYOUT_CREATE_DESCRIPTOR_BUFFER_BIT_EXT;
+			DescriptorSetLayoutCreateInfo.flags = VK_DESCRIPTOR_SET_LAYOUT_CREATE_DESCRIPTOR_BUFFER_BIT_EXT | VK_DESCRIPTOR_SET_LAYOUT_CREATE_PUSH_DESCRIPTOR_BIT_KHR;
 			DescriptorSetLayoutCreateInfo.pNext = nullptr;
 
 			checkSlow(SingleUseUBDescriptorSetLayout == VK_NULL_HANDLE);
@@ -644,7 +636,7 @@ void FVulkanBindlessDescriptorManager::Init()
 				TEXT("StaticShaderBindingLayoutSupport is currently only supported for ray gen shaders in Vulkan."));
 			PushConstants.offset = 0;
 			PushConstants.size = FRHIShaderBindingLayout::MaxUniformBufferEntries * sizeof(uint32);
-			PushConstants.stageFlags = VK_SHADER_STAGE_RAYGEN_BIT_KHR;
+			PushConstants.stageFlags = VK_SHADER_STAGE_ALL;
 			PipelineLayoutCreateInfo.pushConstantRangeCount = 1;
 			PipelineLayoutCreateInfo.pPushConstantRanges = &PushConstants;
 		}
@@ -676,62 +668,23 @@ void FVulkanBindlessDescriptorManager::BindDescriptorBuffers(VkCommandBuffer Com
 	}
 }
 
-void FVulkanBindlessDescriptorManager::RegisterUniformBuffers(FVulkanCommandListContext& Context, VkPipelineBindPoint BindPoint, const FUniformBufferDescriptorArrays& StageUBs)
+void FVulkanBindlessDescriptorManager::RegisterUniformBuffers(FVulkanCommandListContext& Context, VkShaderStageFlags StageFlags, TConstArrayView<VkWriteDescriptorSet> UBDescriptors)
 {
 	checkf(bIsSupported, TEXT("Trying to RegisterUniformBuffers but bindless is not supported!"));
+	check(UBDescriptors.Num() < VulkanBindless::MaxUniformBuffersTotal);
 
-	// :todo-jn: Current uniform buffer layout is a bit wasteful with all the skipped bindings...
-	const uint32 UBDescriptorSize = GetDescriptorTypeSize(Device, VK_DESCRIPTOR_TYPE_UNIFORM_BUFFER);
-	const uint32 BlockDescriptorCount = VulkanBindless::MaxUniformBuffersPerStage * ShaderStage::MaxNumStages;
-	const uint32 BlockSize = BlockDescriptorCount * UBDescriptorSize;
-
-	VkDescriptorBufferBindingInfoEXT LocalBufferBindingInfo[VulkanBindless::MaxNumSets];
-	FMemory::Memcpy(LocalBufferBindingInfo, BufferBindingInfo, VulkanBindless::NumBindlessSets * sizeof(VkDescriptorBufferBindingInfoEXT));
-
-	VkDeviceSize BufferOffsets[VulkanBindless::MaxNumSets];
-	FMemory::Memzero(BufferOffsets);
-
-	uint8* MappedPointer = SingleUseUBAllocator->Alloc(BlockSize, Context,
-		LocalBufferBindingInfo[VulkanBindless::BindlessSingleUseUniformBufferSet], BufferOffsets[VulkanBindless::BindlessSingleUseUniformBufferSet]);
-
-	for (int32 StageIndex = 0; StageIndex < ShaderStage::MaxNumStages; ++StageIndex)
+	VkPushDescriptorSetInfo PushDescriptorSetInfo =
 	{
-		const TArray<VkDescriptorAddressInfoEXT>& DescriptorAddressInfos = StageUBs[StageIndex];
+		.sType = VK_STRUCTURE_TYPE_PUSH_DESCRIPTOR_SET_INFO,
+		.pNext = nullptr,
+		.stageFlags = StageFlags,
+		.layout = BindlessPipelineLayout,
+		.set = VulkanBindless::BindlessSingleUseUniformBufferSet,
+		.descriptorWriteCount = (uint32)UBDescriptors.Num(),
+		.pDescriptorWrites = UBDescriptors.GetData()
+	};
 
-		if (DescriptorAddressInfos.Num())
-		{
-			checkSlow(StageIndex < GetNumStagesForBindPoint(BindPoint));
-			check(DescriptorAddressInfos.Num() <= VulkanBindless::MaxUniformBuffersPerStage);
-			const int32 StageOffset = StageIndex * VulkanBindless::MaxUniformBuffersPerStage;
-
-			for (int32 DescriptorAddressInfoIndex = 0; DescriptorAddressInfoIndex < DescriptorAddressInfos.Num(); DescriptorAddressInfoIndex++)
-			{
-				const VkDescriptorAddressInfoEXT& DescriptorAddressInfo = DescriptorAddressInfos[DescriptorAddressInfoIndex];
-				checkSlow(DescriptorAddressInfo.sType != 0);  // make sure it was filled
-				checkSlow((DescriptorAddressInfo.range % 16) == 0);  // :todo-jn: make sure we don't trip on driver bug, to be removed on next release
-
-				const int32 BindingIndex = StageOffset + DescriptorAddressInfoIndex;
-				VkDeviceSize BindingByteOffset = 0;
-#if UE_BUILD_DEBUG
-				VulkanRHI::vkGetDescriptorSetLayoutBindingOffsetEXT(Device.GetHandle(), SingleUseUBDescriptorSetLayout, BindingIndex, &BindingByteOffset);
-				check(BindingByteOffset == BindingIndex * UBDescriptorSize);
-#else
-				BindingByteOffset = (BindingIndex * UBDescriptorSize);
-#endif
-
-				VkDescriptorGetInfoEXT Info;
-				ZeroVulkanStruct(Info, VK_STRUCTURE_TYPE_DESCRIPTOR_GET_INFO_EXT);
-				Info.type = VK_DESCRIPTOR_TYPE_UNIFORM_BUFFER;
-				Info.data.pUniformBuffer = &DescriptorAddressInfo;
-				VulkanRHI::vkGetDescriptorEXT(Device.GetHandle(), &Info, UBDescriptorSize, &MappedPointer[BindingByteOffset]);
-			}
-		}
-	}
-
-	// todo-jn: cache these states and only repeat when necessary
-	const VkCommandBuffer CommandBufferHandle = Context.GetCommandBuffer().GetHandle();
-	VulkanRHI::vkCmdBindDescriptorBuffersEXT(CommandBufferHandle, VulkanBindless::MaxNumSets, LocalBufferBindingInfo);
-	VulkanRHI::vkCmdSetDescriptorBufferOffsetsEXT(CommandBufferHandle, BindPoint, BindlessPipelineLayout, 0u, VulkanBindless::MaxNumSets, BufferIndices, BufferOffsets);
+	VulkanRHI::vkCmdPushDescriptorSet2KHR(Context.GetCommandBuffer().GetHandle(), &PushDescriptorSetInfo);
 }
 
 FRHIDescriptorHandle FVulkanBindlessDescriptorManager::AllocateDescriptor(VkDescriptorType DescriptorType)
@@ -959,11 +912,3 @@ void FVulkanBindlessDescriptorManager::UpdateAccelerationStructure(const FVulkan
 		UpdateDescriptor(Contexts, DescriptorHandle, VK_DESCRIPTOR_TYPE_ACCELERATION_STRUCTURE_KHR, DescriptorData, bImmediateUpdate);
 	}
 }
-
-void FVulkanBindlessDescriptorManager::UpdateUBAllocator()
-{
-	if (bIsSupported && SingleUseUBAllocator)
-	{
-		SingleUseUBAllocator->UpdateBlocks();
-	}
-}
diff --git a/Engine/Source/Runtime/VulkanRHI/Private/VulkanBindlessDescriptorManager.h b/Engine/Source/Runtime/VulkanRHI/Private/VulkanBindlessDescriptorManager.h
--- a/Engine/Source/Runtime/VulkanRHI/Private/VulkanBindlessDescriptorManager.h
+++ b/Engine/Source/Runtime/VulkanRHI/Private/VulkanBindlessDescriptorManager.h
@@ -19,8 +19,6 @@ public:
 	FVulkanBindlessDescriptorManager(FVulkanDevice& InDevice);
 	~FVulkanBindlessDescriptorManager();
 
-	typedef TStaticArray<TArray<VkDescriptorAddressInfoEXT>, ShaderStage::MaxNumStages> FUniformBufferDescriptorArrays;
-
 	void Init();
 	void Deinit();
 
@@ -43,7 +41,6 @@ public:
 	void BindDescriptorBuffers(VkCommandBuffer CommandBuffer, VkPipelineStageFlags SupportedStages);
 
 	FRHIDescriptorHandle AllocateDescriptor(VkDescriptorType DescriptorType);
-	FRHIDescriptorHandle AllocateDescriptor(VkDescriptorType DescriptorType, VkImageView VulkanImage, bool bIsDepthStencil);
 	FRHIDescriptorHandle AllocateDescriptor(VkDescriptorType DescriptorType, VkDeviceAddress BufferAddress, VkDeviceSize BufferSize);
 	FRHIDescriptorHandle AllocateDescriptor(VkDescriptorType DescriptorType, FVulkanBuffer* Buffer, uint32 ExtraOffset);
 
@@ -60,9 +57,7 @@ public:
 	void UpdateTexelBuffer(const FVulkanContextArray& Contexts, FRHIDescriptorHandle DescriptorHandle, VkDescriptorType DescriptorType, const VkBufferViewCreateInfo& ViewInfo, bool bImmediateUpdate = true);
 	void UpdateAccelerationStructure(const FVulkanContextArray& Contexts, FRHIDescriptorHandle DescriptorHandle, VkAccelerationStructureKHR AccelerationStructure, bool bImmediateUpdate = true);
 
-	void RegisterUniformBuffers(FVulkanCommandListContext& Context, VkPipelineBindPoint BindPoint, const FUniformBufferDescriptorArrays& StageUBs);
-
-	void UpdateUBAllocator();
+	void RegisterUniformBuffers(FVulkanCommandListContext& Context, VkShaderStageFlags StageFlags, TConstArrayView<VkWriteDescriptorSet> UBDescriptors);
 
 private:
 	FVulkanDevice& Device;
@@ -97,7 +92,6 @@ private:
 	BindlessSetState BindlessSetStates[VulkanBindless::NumBindlessSets];
 
 	VkDescriptorSetLayout SingleUseUBDescriptorSetLayout = VK_NULL_HANDLE;
-	VulkanRHI::FTempBlockAllocator* SingleUseUBAllocator = nullptr;
 
 	VkDescriptorBufferBindingInfoEXT BufferBindingInfo[VulkanBindless::NumBindlessSets];
 	uint32_t BufferIndices[VulkanBindless::MaxNumSets];
diff --git a/Engine/Source/Runtime/VulkanRHI/Private/VulkanExtensions.cpp b/Engine/Source/Runtime/VulkanRHI/Private/VulkanExtensions.cpp
--- a/Engine/Source/Runtime/VulkanRHI/Private/VulkanExtensions.cpp
+++ b/Engine/Source/Runtime/VulkanRHI/Private/VulkanExtensions.cpp
@@ -1596,12 +1596,25 @@ public:
 		}
 	}
 
-	virtual void PreCreateDevice(VkDeviceCreateInfo& DeviceCreateInfo) override final
+	virtual void PostPhysicalDeviceProperties(FOptionalVulkanDeviceExtensions& ExtensionFlags) override final
 	{
 		if (bRequirementsPassed)
 		{
-			AddToPNext(DeviceCreateInfo, DescriptorBufferFeatures);
+			const VkPhysicalDeviceDescriptorBufferPropertiesEXT& DescriptorBufferProperties = GetDeviceExtensionProperties().DescriptorBufferProps;
 
+			if (DescriptorBufferProperties.bufferlessPushDescriptors == VK_FALSE)
+			{
+				UE_LOG(LogVulkanRHI, Display, TEXT("Disabling Vulkan Descriptor Buffers because bufferless isn't supported."));
+
+				ExtensionFlags.HasEXTDescriptorBuffer = false;
+			}
+		}
+	}
+
+	virtual void PreCreateDevice(VkDeviceCreateInfo& DeviceCreateInfo) override final
+	{
+		if (bRequirementsPassed)
+		{
 			const VkPhysicalDeviceDescriptorBufferPropertiesEXT& DescriptorBufferProperties = GetDeviceExtensionProperties().DescriptorBufferProps;
 
 			UE_LOG(LogVulkanRHI, Display, TEXT("Enabling Vulkan Descriptor Buffers with: ")
@@ -1617,6 +1630,8 @@ public:
 				DescriptorBufferProperties.maxSamplerDescriptorBufferRange, DescriptorBufferProperties.maxResourceDescriptorBufferRange,
 				DescriptorBufferProperties.descriptorBufferAddressSpaceSize, DescriptorBufferProperties.descriptorBufferOffsetAlignment,
 				DescriptorBufferProperties.samplerDescriptorSize);
+
+			AddToPNext(DeviceCreateInfo, DescriptorBufferFeatures);
 		}
 	}
 
@@ -2251,7 +2266,6 @@ public:
 	VkPhysicalDeviceHostImageCopyProperties HostImageCopyProperties;
 };
 
-
 template <typename ExtensionType>
 static void FlagExtensionSupport(const TArray<VkExtensionProperties>& ExtensionProperties, TArray<TUniquePtr<ExtensionType>>& UEExtensions, uint32 ApiVersion, const TCHAR* ExtensionTypeName)
 {
@@ -2305,6 +2319,7 @@ FVulkanDeviceExtensionArray FVulkanDeviceExtension::GetUESupportedDeviceExtensio
 	ADD_SIMPLE_EXTENSION(VK_QCOM_RENDER_PASS_SHADER_RESOLVE_EXTENSION_NAME, VULKAN_SUPPORTS_QCOM_RENDERPASS_SHADER_RESOLVE, VULKAN_EXTENSION_NOT_PROMOTED, DEVICE_EXT_FLAG_SETTER(HasQcomRenderPassShaderResolve));
 	ADD_SIMPLE_EXTENSION(VK_EXT_TOOLING_INFO_EXTENSION_NAME, VULKAN_HAS_DEBUGGING_ENABLED || VULKAN_ENABLE_DRAW_MARKERS, VK_API_VERSION_1_3,  DEVICE_EXT_FLAG_SETTER(HasEXTToolingInfo));
 	ADD_SIMPLE_EXTENSION(VK_KHR_PIPELINE_LIBRARY_EXTENSION_NAME,         VULKAN_EXTENSION_ENABLED,             VULKAN_EXTENSION_NOT_PROMOTED, DEVICE_EXT_FLAG_SETTER(HasKHRPipelineLibrary));
+	ADD_SIMPLE_EXTENSION(VK_KHR_PUSH_DESCRIPTOR_EXTENSION_NAME,          VULKAN_EXTENSION_ENABLED,             VK_API_VERSION_1_4,            DEVICE_EXT_FLAG_SETTER(HasKHRPushDescriptors));
 
 	// Externally activated extensions (supported by the engine, but enabled externally by plugin or other) :
 	ADD_EXTERNAL_EXTENSION(VK_EXT_DEBUG_MARKER_EXTENSION_NAME, VULKAN_HAS_DEBUGGING_ENABLED || VULKAN_ENABLE_DRAW_MARKERS, VULKAN_EXTENSION_NOT_PROMOTED, nullptr);
@@ -2314,6 +2329,7 @@ FVulkanDeviceExtensionArray FVulkanDeviceExtension::GetUESupportedDeviceExtensio
 	ADD_CUSTOM_EXTENSION(FVulkanKHRDriverPropertiesExtension);
 	ADD_CUSTOM_EXTENSION(FVulkanKHRMaintenance4Extension);
 	ADD_CUSTOM_EXTENSION(FVulkanKHRMaintenance5Extension);
+	ADD_CUSTOM_EXTENSION(FVulkanKHRMaintenance6Extension);
 	ADD_CUSTOM_EXTENSION(FVulkanKHRMaintenance7Extension);
 	ADD_CUSTOM_EXTENSION(FVulkanShaderAtomicInt64Extension);
 	ADD_CUSTOM_EXTENSION(FVulkanShaderImageAtomicInt64Extension);
diff --git a/Engine/Source/Runtime/VulkanRHI/Private/VulkanPipelineState.cpp b/Engine/Source/Runtime/VulkanRHI/Private/VulkanPipelineState.cpp
--- a/Engine/Source/Runtime/VulkanRHI/Private/VulkanPipelineState.cpp
+++ b/Engine/Source/Runtime/VulkanRHI/Private/VulkanPipelineState.cpp
@@ -238,13 +238,21 @@ void FVulkanComputePipelineDescriptorState::UpdateBindlessDescriptors(FVulkanCom
 	check(DSWriteContainer.DescriptorBufferInfo.Num() == DSWriteContainer.DescriptorWrites.Num());
 	check(DSWriteContainer.DescriptorImageInfo.Num() == 0);
 
-	FVulkanBindlessDescriptorManager::FUniformBufferDescriptorArrays StageUBs;
-
 	const FVulkanShaderHeader& Header = ComputePipeline->GetShaderCodeHeader();
+	checkf(Header.NumBoundUniformBuffers < VulkanBindless::MaxUniformBuffersPerStage, TEXT("Compute shader has more uniform buffers than is permitted in bindless."));
+	check(Header.NumBoundUniformBuffers == DSWriter[ShaderStage::EStage::Compute].NumWrites);
+
+	if (Header.NumBoundUniformBuffers == 0)
+	{
+		return;
+	}
+
+	// Create a local copy of the writes
+	TArray<VkWriteDescriptorSet, TInlineAllocator<VulkanBindless::MaxUniformBuffersPerStage>> UBWrites;
+	UBWrites.SetNumZeroed(Header.NumBoundUniformBuffers);
+	FMemory::Memcpy(UBWrites.GetData(), DSWriter[ShaderStage::EStage::Compute].WriteDescriptors, Header.NumBoundUniformBuffers * sizeof(VkWriteDescriptorSet));
 
-	TArray<VkDescriptorAddressInfoEXT>& DescriptorAddressInfos = StageUBs[ShaderStage::EStage::Compute];
-	DescriptorAddressInfos.SetNumZeroed(Header.NumBoundUniformBuffers);
-	uint32 UBIndex = 0;
+	VkDescriptorBufferInfo ConstantBufferInfo;
 
 	// UBs are currently set from a fresh batch of descriptors for every call, so ignore PackedUniformBuffersDirty
 	check(PackedUniformBuffersMask <= 1);
@@ -252,44 +260,29 @@ void FVulkanComputePipelineDescriptorState::UpdateBindlessDescriptors(FVulkanCom
 	{
 		const FPackedUniformBuffers::FPackedBuffer& StagedUniformBuffer = PackedUniformBuffers.GetBuffer();
 		const int32 UBSize = StagedUniformBuffer.Num();
-		const int32 BindingIndex = 0;
 		const VkDeviceSize UBOffsetAlignment = Device.GetLimits().minUniformBufferOffsetAlignment;
 
 		VulkanRHI::FVulkanAllocation TempAllocation;
-		uint8* MappedPointer = Device.GetTempBlockAllocator().Alloc(UBSize, UBOffsetAlignment, Context, TempAllocation, &DescriptorAddressInfos[BindingIndex]);
+		uint8* MappedPointer = Device.GetTempBlockAllocator().Alloc(UBSize, UBOffsetAlignment, Context, TempAllocation, nullptr);
 		FMemory::Memcpy(MappedPointer, StagedUniformBuffer.GetData(), UBSize);
 
-		PackedUniformBuffersDirty = 0;
-		++UBIndex;
-	}
-
-	for (;UBIndex < Header.NumBoundUniformBuffers; ++UBIndex)
-	{
-		VkDescriptorAddressInfoEXT& DescriptorAddressInfo = DescriptorAddressInfos[UBIndex];
-		check(DescriptorAddressInfo.sType == 0);
-
-		VkWriteDescriptorSet& WriteDescriptorSet = DSWriter[ShaderStage::EStage::Compute].WriteDescriptors[UBIndex];
-		check(WriteDescriptorSet.dstBinding == UBIndex);
+		VkWriteDescriptorSet& WriteDescriptorSet = UBWrites[0];
+		check(WriteDescriptorSet.dstBinding == 0);
 		check(WriteDescriptorSet.dstArrayElement == 0);
 		check(WriteDescriptorSet.descriptorCount == 1);
 		check(WriteDescriptorSet.descriptorType == VK_DESCRIPTOR_TYPE_UNIFORM_BUFFER);
-		checkf(WriteDescriptorSet.pBufferInfo && WriteDescriptorSet.pBufferInfo->buffer, 
-			TEXT("Missing uniform buffer binding for [%s] at index [%d] of shader [%s]."),
-			*ComputePipeline->GetComputeShader()->GetUniformBufferName(UBIndex), UBIndex,
-			ComputePipeline->GetComputeShader()->GetShaderName());
-
-		VkBufferDeviceAddressInfo BufferInfo;
-		ZeroVulkanStruct(BufferInfo, VK_STRUCTURE_TYPE_BUFFER_DEVICE_ADDRESS_INFO);
-		BufferInfo.buffer = WriteDescriptorSet.pBufferInfo->buffer;
-		VkDeviceAddress BufferAddress = VulkanRHI::vkGetBufferDeviceAddressKHR(Device.GetHandle(), &BufferInfo);
-
-		DescriptorAddressInfo.sType = VK_STRUCTURE_TYPE_DESCRIPTOR_ADDRESS_INFO_EXT;
-		DescriptorAddressInfo.address = BufferAddress + WriteDescriptorSet.pBufferInfo->offset;
-		DescriptorAddressInfo.range = WriteDescriptorSet.pBufferInfo->range;
+		
+		WriteDescriptorSet.pBufferInfo = &ConstantBufferInfo;
+
+		ConstantBufferInfo.buffer = TempAllocation.GetBufferHandle();
+		ConstantBufferInfo.offset = TempAllocation.Offset;
+		ConstantBufferInfo.range = UBSize;
+
+		PackedUniformBuffersDirty = 0;
 	}
 
 	// Send to descriptor manager
-	Device.GetBindlessDescriptorManager()->RegisterUniformBuffers(Context, VK_PIPELINE_BIND_POINT_COMPUTE, StageUBs);
+	Device.GetBindlessDescriptorManager()->RegisterUniformBuffers(Context, VK_SHADER_STAGE_COMPUTE_BIT, UBWrites);
 }
 
 FVulkanGraphicsPipelineDescriptorState::FVulkanGraphicsPipelineDescriptorState(FVulkanDevice& InDevice, FVulkanGraphicsPipelineState* InGfxPipeline)
@@ -471,71 +464,85 @@ void FVulkanGraphicsPipelineDescriptorState::UpdateBindlessDescriptors(FVulkanCo
 	check(DSWriteContainer.DescriptorBufferInfo.Num() == DSWriteContainer.DescriptorWrites.Num());
 	check(DSWriteContainer.DescriptorImageInfo.Num() == 0);
 
+	checkf(!GfxPipeline->GetVulkanShader(SF_Vertex) || !GfxPipeline->GetVulkanShader(SF_Geometry) ||
+		GfxPipeline->GetVulkanShader(SF_Vertex)->GetCodeHeader().NumBoundUniformBuffers +
+		GfxPipeline->GetVulkanShader(SF_Geometry)->GetCodeHeader().NumBoundUniformBuffers <
+		VulkanBindless::MaxUniformBuffersPerStage, TEXT("Sum of Vertex UB count and Gemetry UB count exceeds bindless limit."));
+
+	checkf(!GfxPipeline->GetVulkanShader(SF_Mesh) || !GfxPipeline->GetVulkanShader(SF_Amplification) ||
+		GfxPipeline->GetVulkanShader(SF_Mesh)->GetCodeHeader().NumBoundUniformBuffers +
+		GfxPipeline->GetVulkanShader(SF_Amplification)->GetCodeHeader().NumBoundUniformBuffers <
+		VulkanBindless::MaxUniformBuffersPerStage, TEXT("Sum of Mesh UB count and Task UB count exceeds bindless limit."));
+
 	const VkDeviceSize UBOffsetAlignment = Device.GetLimits().minUniformBufferOffsetAlignment;
 
-	FVulkanBindlessDescriptorManager::FUniformBufferDescriptorArrays StageUBs;
+	TArray<VkWriteDescriptorSet, TInlineAllocator<VulkanBindless::MaxUniformBuffersTotal>> UBCompactedWrites;
+	TStaticArray<VkDescriptorBufferInfo, ShaderStage::NumGraphicsStages> ConstantBufferInfos;
+
+	VkShaderStageFlags StageFlags = 0;
 
-	// Process updates
-	{
 #if VULKAN_ENABLE_AGGRESSIVE_STATS
-		SCOPE_CYCLE_COUNTER(STAT_VulkanApplyPackedUniformBuffers);
+	SCOPE_CYCLE_COUNTER(STAT_VulkanApplyPackedUniformBuffers);
 #endif
-		for (int32 Stage = 0; Stage < ShaderStage::NumGraphicsStages; ++Stage)
+	for (int32 Stage = 0; Stage < ShaderStage::NumGraphicsStages; ++Stage)
+	{
+		const EShaderFrequency Frequency = GetFrequencyForGfxStage((ShaderStage::EStage)Stage);
+		const FVulkanShader* VulkanShader = GfxPipeline->GetVulkanShader(Frequency);
+		if (!VulkanShader)
 		{
-			const FVulkanShader* VulkanShader = GfxPipeline->GetVulkanShader(GetFrequencyForGfxStage((ShaderStage::EStage)Stage));
-			if (!VulkanShader)
-			{
-				continue;
-			}
+			continue;
+		}
 
-			const FVulkanShaderHeader& Header = VulkanShader->GetCodeHeader();
+		StageFlags |= UEFrequencyToVKStageBit(Frequency);
 
-			TArray<VkDescriptorAddressInfoEXT>& DescriptorAddressInfos = StageUBs[Stage];
-			DescriptorAddressInfos.SetNumZeroed(Header.NumBoundUniformBuffers);
-			uint32 UBIndex = 0;
+		const FVulkanShaderHeader& Header = VulkanShader->GetCodeHeader();
+
+		// Move the writes to compacted binding indices
+		const int32 BindingOffset = VulkanBindless::GetOffsetForFrequency(Frequency);
+		for (int32 UBIndex = 0; UBIndex < (int32)Header.NumBoundUniformBuffers; ++UBIndex)
+		{
+			const VkWriteDescriptorSet& WriteDescriptorSet = DSWriter[Stage].WriteDescriptors[UBIndex];
+			check(WriteDescriptorSet.dstBinding == UBIndex);
+			check(WriteDescriptorSet.dstArrayElement == 0);
+			check(WriteDescriptorSet.descriptorCount == 1);
+			check(WriteDescriptorSet.descriptorType == VK_DESCRIPTOR_TYPE_UNIFORM_BUFFER);
+
+			VkWriteDescriptorSet& CompactedWrite = UBCompactedWrites.Add_GetRef(WriteDescriptorSet);
+			if (BindingOffset < 0)
+			{
+				CompactedWrite.dstBinding = VulkanBindless::MaxUniformBuffersTotal - UBIndex - 1;
+			}
+			else
+			{
+				CompactedWrite.dstBinding = BindingOffset + UBIndex;
+			}
 
 			// UBs are currently set from a fresh batch of descriptors for every call, so ignore PackedUniformBuffersDirty
 			check(PackedUniformBuffersMask[Stage] <= 1);
-			if (PackedUniformBuffersMask[Stage] != 0)
+			if ((UBIndex == 0) && (PackedUniformBuffersMask[Stage] != 0))
 			{
 				const FPackedUniformBuffers::FPackedBuffer& StagedUniformBuffer = PackedUniformBuffers[Stage].GetBuffer();
 				const int32 UBSize = StagedUniformBuffer.Num();
-				const int32 BindingIndex = 0;
 
 				VulkanRHI::FVulkanAllocation TempAllocation;
-				uint8* MappedPointer = Device.GetTempBlockAllocator().Alloc(UBSize, UBOffsetAlignment, Context, TempAllocation, &DescriptorAddressInfos[BindingIndex]);
+				uint8* MappedPointer = Device.GetTempBlockAllocator().Alloc(UBSize, UBOffsetAlignment, Context, TempAllocation, nullptr);
 				FMemory::Memcpy(MappedPointer, StagedUniformBuffer.GetData(), UBSize);
 
-				PackedUniformBuffersDirty[Stage] = 0;
-				++UBIndex;
-			}
+				VkDescriptorBufferInfo& ConstantBufferInfo = ConstantBufferInfos[Stage];
 
-			for (; UBIndex < Header.NumBoundUniformBuffers; ++UBIndex)
-			{
-				VkDescriptorAddressInfoEXT& DescriptorAddressInfo = DescriptorAddressInfos[UBIndex];
-				check(DescriptorAddressInfo.sType == 0);
-
-				VkWriteDescriptorSet& WriteDescriptorSet = DSWriter[Stage].WriteDescriptors[UBIndex];
-				check(WriteDescriptorSet.dstBinding == UBIndex);
-				check(WriteDescriptorSet.dstArrayElement == 0);
-				check(WriteDescriptorSet.descriptorCount == 1);
-				check(WriteDescriptorSet.descriptorType == VK_DESCRIPTOR_TYPE_UNIFORM_BUFFER);
-				checkSlow(WriteDescriptorSet.pBufferInfo);
-
-				VkBufferDeviceAddressInfo BufferInfo;
-				ZeroVulkanStruct(BufferInfo, VK_STRUCTURE_TYPE_BUFFER_DEVICE_ADDRESS_INFO);
-				BufferInfo.buffer = WriteDescriptorSet.pBufferInfo->buffer;
-				VkDeviceAddress BufferAddress = VulkanRHI::vkGetBufferDeviceAddressKHR(Device.GetHandle(), &BufferInfo);
-
-				DescriptorAddressInfo.sType = VK_STRUCTURE_TYPE_DESCRIPTOR_ADDRESS_INFO_EXT;
-				DescriptorAddressInfo.address = BufferAddress + WriteDescriptorSet.pBufferInfo->offset;
-				DescriptorAddressInfo.range = WriteDescriptorSet.pBufferInfo->range;
+				CompactedWrite.pBufferInfo = &ConstantBufferInfo;
+
+				ConstantBufferInfo.buffer = TempAllocation.GetBufferHandle();
+				ConstantBufferInfo.offset = TempAllocation.Offset;
+				ConstantBufferInfo.range = UBSize;
+
+				PackedUniformBuffersDirty[Stage] = 0;
 			}
 		}
 	}
 
 	// Send to descriptor manager
-	Device.GetBindlessDescriptorManager()->RegisterUniformBuffers(Context, VK_PIPELINE_BIND_POINT_GRAPHICS, StageUBs);
+	Device.GetBindlessDescriptorManager()->RegisterUniformBuffers(Context, StageFlags, UBCompactedWrites);
 }
 
 
diff --git a/Engine/Source/Runtime/VulkanRHI/Private/VulkanRHI.cpp b/Engine/Source/Runtime/VulkanRHI/Private/VulkanRHI.cpp
--- a/Engine/Source/Runtime/VulkanRHI/Private/VulkanRHI.cpp
+++ b/Engine/Source/Runtime/VulkanRHI/Private/VulkanRHI.cpp
@@ -1181,7 +1181,6 @@ void FVulkanDynamicRHI::RHIEndFrame_RenderThread(FRHICommandListImmediate& RHICm
 
 		Context.Device.GetPipelineStateCache()->TickLRU();
 
-		Context.Device.GetBindlessDescriptorManager()->UpdateUBAllocator();
 		Context.Device.GetTempBlockAllocator().UpdateBlocks();
 	});
 
diff --git a/Engine/Source/Runtime/VulkanRHI/Private/VulkanRayTracing.cpp b/Engine/Source/Runtime/VulkanRHI/Private/VulkanRayTracing.cpp
--- a/Engine/Source/Runtime/VulkanRHI/Private/VulkanRayTracing.cpp
+++ b/Engine/Source/Runtime/VulkanRHI/Private/VulkanRayTracing.cpp
@@ -979,7 +979,7 @@ void FVulkanRayTracingShaderTable::SetLocalShaderParameters(EShaderFrequency Fre
 
 	checkfSlow(OffsetWithinRecord % 4 == 0, TEXT("SBT record parameters must be written on DWORD-aligned boundary"));
 	checkfSlow(InDataSize % 4 == 0, TEXT("SBT record parameters must be DWORD-aligned"));
-	checkf(OffsetWithinRecord + InDataSize <= Alloc.Region.stride ? Alloc.Region.stride : Alloc.Region.size, TEXT("SBT record write request is out of bounds"));
+	checkf((OffsetWithinRecord + InDataSize) <= (Alloc.Region.stride ? Alloc.Region.stride : Alloc.Region.size), TEXT("SBT record write request is out of bounds"));
 	checkf((RecordIndex == 0) || (Alloc.Region.stride != 0), TEXT("Attempting to index a record in a region without stride"));
 
 	const uint32 WriteOffset = HandleSizeAligned + (Alloc.Region.stride * RecordIndex) + OffsetWithinRecord;
@@ -1916,41 +1916,32 @@ static void SetRayGenResources(FVulkanDevice& Device, FVulkanCommandListContext&
 
 	// Uniform buffers
 	{
-		uint32 NumSkippedSlots = 0;
-		FVulkanBindlessDescriptorManager::FUniformBufferDescriptorArrays StageUBs;
+		TArray<VkWriteDescriptorSet, TInlineAllocator<VulkanBindless::MaxUniformBuffersPerStage>> UBWrites;
+		TArray<VkDescriptorBufferInfo, TInlineAllocator<VulkanBindless::MaxUniformBuffersPerStage>> UBBufferInfos;
+
 		const uint32 MaxUniformBuffers = UE_ARRAY_COUNT(InGlobalResourceBindings.UniformBuffers);
 		for (uint32 UBIndex = 0; UBIndex < MaxUniformBuffers; ++UBIndex)
 		{
 			const FVulkanUniformBuffer* UniformBuffer = ResourceCast(InGlobalResourceBindings.UniformBuffers[UBIndex]);
 			if (UniformBuffer)
 			{
-				if (NumSkippedSlots > 0)
-				{
-					UE_LOG(LogRHI, Warning, TEXT("Skipping %u Uniform Buffer bindings, this isn't normal!"), NumSkippedSlots);
-
-					for (uint32 SkipIndex = 0; SkipIndex < NumSkippedSlots; ++SkipIndex)
-					{
-						VkDescriptorAddressInfoEXT& DescriptorAddressInfo = StageUBs[ShaderStage::EStage::RayGen].AddZeroed_GetRef();
-						DescriptorAddressInfo.sType = VK_STRUCTURE_TYPE_DESCRIPTOR_ADDRESS_INFO_EXT;
-					}
-
-					NumSkippedSlots = 0;
-				}
-				
-				VkDescriptorAddressInfoEXT& DescriptorAddressInfo = StageUBs[ShaderStage::EStage::RayGen].AddZeroed_GetRef();
-				DescriptorAddressInfo.sType = VK_STRUCTURE_TYPE_DESCRIPTOR_ADDRESS_INFO_EXT;
-				DescriptorAddressInfo.address = UniformBuffer->GetDeviceAddress();
-				DescriptorAddressInfo.range = UniformBuffer->GetSize();
+				VkDescriptorBufferInfo& BufferInfo = UBBufferInfos.AddZeroed_GetRef();
+				BufferInfo.buffer = UniformBuffer->GetBufferHandle();
+				BufferInfo.offset = UniformBuffer->GetOffset();
+				BufferInfo.range = UniformBuffer->GetSize();
+
+				VkWriteDescriptorSet& Write = UBWrites.AddZeroed_GetRef();
+				Write.sType = VK_STRUCTURE_TYPE_WRITE_DESCRIPTOR_SET;
+				Write.dstBinding = UBIndex;
+				Write.descriptorType = VK_DESCRIPTOR_TYPE_UNIFORM_BUFFER;
+				Write.descriptorCount = 1;
+				Write.pBufferInfo = &BufferInfo;
 
 				UniformBuffers.AddUnique(UniformBuffer);
 			}
-			else
-			{
-				// :todo-jn: There might be unused indices (see USE_INSTANCE_CULLING_DATA issue), just skip them with a warning for now.
-				NumSkippedSlots++;
-			}
 		}
-		Device.GetBindlessDescriptorManager()->RegisterUniformBuffers(Context, VK_PIPELINE_BIND_POINT_RAY_TRACING_KHR, StageUBs);
+
+		Device.GetBindlessDescriptorManager()->RegisterUniformBuffers(Context, VK_SHADER_STAGE_RAYGEN_BIT_KHR, UBWrites);
 	}
 
 	// Static uniform buffers
diff --git a/Engine/Source/Runtime/VulkanRHI/Public/VulkanCommon.h b/Engine/Source/Runtime/VulkanRHI/Public/VulkanCommon.h
--- a/Engine/Source/Runtime/VulkanRHI/Public/VulkanCommon.h
+++ b/Engine/Source/Runtime/VulkanRHI/Public/VulkanCommon.h
@@ -102,6 +102,7 @@ namespace ShaderStage
 namespace VulkanBindless
 {
 	static constexpr uint32 MaxUniformBuffersPerStage = 16;
+	static constexpr uint32 MaxUniformBuffersTotal = 32;  // Guaranteed to be min 32 by Vulkan 1.4
 
 	enum EDescriptorSets
 	{
@@ -138,6 +139,38 @@ namespace VulkanBindless
 		// Total number of descriptor sets used in a bindless pipeline
 		MaxNumSets = NumBindlessSets + 1
 	};
+
+	inline int32 GetOffsetForFrequency(EShaderFrequency Frequency)
+	{
+		// Stages that run alone or that are likely to contain more UBs start at 0
+		// The second stage follows starting at +MaxUniformBuffersPerStage
+		// The third stage returns -1 to signal descriptors are allocated from the end 
+		// The sum of second and third stage must never go beyond MaxUniformBuffersPerStage
+
+		switch (Frequency)
+		{
+		case SF_Pixel:			return 0;
+
+		case SF_Vertex:			return MaxUniformBuffersPerStage;
+		case SF_Geometry:		return -1;
+
+		case SF_Mesh:			return MaxUniformBuffersPerStage;
+		case SF_Amplification:	return -1;
+
+		case SF_Compute:		return 0;
+
+		case SF_RayGen:			
+		case SF_RayMiss:
+		case SF_RayHitGroup:
+		case SF_RayCallable:	return 0;
+		default:
+			checkf(false, TEXT("Invalid shader frequency %d"), (int32)Frequency);
+			break;
+		}
+
+		return 0;
+	}
+
 };
 
 DECLARE_LOG_CATEGORY_EXTERN(LogVulkan, Display, All);
