VulkanRHI: Add names to bindless descriptor set layouts.  Match push constant from pipeline layout when creating shader objects.

#rnx

diff --git a/Engine/Source/Runtime/VulkanRHI/Private/VulkanBindlessDescriptorManager.cpp b/Engine/Source/Runtime/VulkanRHI/Private/VulkanBindlessDescriptorManager.cpp
--- a/Engine/Source/Runtime/VulkanRHI/Private/VulkanBindlessDescriptorManager.cpp
+++ b/Engine/Source/Runtime/VulkanRHI/Private/VulkanBindlessDescriptorManager.cpp
@@ -412,7 +412,7 @@ void FVulkanBindlessDescriptorManager::Init()
 		}
 
 		// Fill the DescriptorSetLayout for a BindlessSetState
-		auto CreateDescriptorSetLayout = [&](const BindlessSetState& State)
+		auto CreateDescriptorSetLayout = [&, &DeviceRef=Device](const BindlessSetState& State)
 		{
 			if (State.DescriptorTypes.Num() == 0)
 			{
@@ -470,6 +470,8 @@ void FVulkanBindlessDescriptorManager::Init()
 				VulkanRHI::vkGetDescriptorSetLayoutBindingOffsetEXT(DeviceHandle, DescriptorSetLayout, 0u, &BindingOffset);
 				check(BindingOffset == 0);
 
+				VULKAN_SET_DEBUG_NAME(DeviceRef, VK_OBJECT_TYPE_DESCRIPTOR_SET_LAYOUT, DescriptorSetLayout, TEXT("BindlessArray(%s)"), VK_TYPE_TO_STRING(VkDescriptorType, DescriptorSetLayoutBinding.descriptorType));
+
 				return DescriptorSetLayout;
 			}
 		};
@@ -606,6 +608,7 @@ void FVulkanBindlessDescriptorManager::Init()
 
 			checkSlow(SingleUseUBDescriptorSetLayout == VK_NULL_HANDLE);
 			VERIFYVULKANRESULT(VulkanRHI::vkCreateDescriptorSetLayout(DeviceHandle, &DescriptorSetLayoutCreateInfo, VULKAN_CPU_ALLOCATOR, &SingleUseUBDescriptorSetLayout));
+			VULKAN_SET_DEBUG_NAME(Device, VK_OBJECT_TYPE_DESCRIPTOR_SET_LAYOUT, SingleUseUBDescriptorSetLayout, TEXT("UBPushDescriptors(%d)"), NumTotalBindings);
 		}
 
 		checkf(TotalResourceDescriptorBufferSize < DescriptorBufferProperties.resourceDescriptorBufferAddressSpaceSize,
diff --git a/Engine/Source/Runtime/VulkanRHI/Private/VulkanShaderObjectManager.cpp b/Engine/Source/Runtime/VulkanRHI/Private/VulkanShaderObjectManager.cpp
--- a/Engine/Source/Runtime/VulkanRHI/Private/VulkanShaderObjectManager.cpp
+++ b/Engine/Source/Runtime/VulkanRHI/Private/VulkanShaderObjectManager.cpp
@@ -65,6 +65,18 @@ VkShaderEXT FVulkanShaderObjectManager::CreateShaderObject(FVulkanShader& Vulkan
 		.pSpecializationInfo = nullptr
 	};
 
+	VkPushConstantRange PushConstants;
+	if (RHIGetStaticShaderBindingLayoutSupport(GMaxRHIShaderPlatform) != ERHIStaticShaderBindingLayoutSupport::Unsupported)
+	{
+		checkf(RHIGetStaticShaderBindingLayoutSupport(GMaxRHIShaderPlatform) == ERHIStaticShaderBindingLayoutSupport::RayTracingOnly,
+			TEXT("StaticShaderBindingLayoutSupport is currently only supported for ray gen shaders in Vulkan."));
+		PushConstants.offset = 0;
+		PushConstants.size = FRHIShaderBindingLayout::MaxUniformBufferEntries * sizeof(uint32);
+		PushConstants.stageFlags = VK_SHADER_STAGE_ALL;
+		ShaderCreateInfo.pushConstantRangeCount = 1;
+		ShaderCreateInfo.pPushConstantRanges = &PushConstants;
+	}
+
 	// :todo-jn: Task shaders never used at the moment
 	if (Stage == VK_SHADER_STAGE_MESH_BIT_EXT)
 	{
