VulkanRHI: Add functions for VK_EXT_device_generated_commands and keep a copy of the properties for runtime checks.

#rnx

diff --git a/Engine/Source/Runtime/VulkanRHI/Private/VulkanCommandWrappers.h b/Engine/Source/Runtime/VulkanRHI/Private/VulkanCommandWrappers.h
--- a/Engine/Source/Runtime/VulkanRHI/Private/VulkanCommandWrappers.h
+++ b/Engine/Source/Runtime/VulkanRHI/Private/VulkanCommandWrappers.h
@@ -300,6 +300,16 @@ struct FWrapLayer
 	static void CmdSetSampleMaskEXT(VkResult Result, VkCommandBuffer CommandBuffer, VkSampleCountFlagBits Samples, const VkSampleMask* SampleMask) VULKAN_LAYER_BODY
 	static void CmdSetFragmentShadingRateKHR(VkResult Result, VkCommandBuffer CommandBuffer, const VkExtent2D* FragmentSize, const VkFragmentShadingRateCombinerOpKHR combinerOps[2]) VULKAN_LAYER_BODY
 
+	static void CmdExecuteGeneratedCommandsEXT(VkResult Result, VkCommandBuffer CommandBuffer, VkBool32 IsPreprocessed, const VkGeneratedCommandsInfoEXT* GeneratedCommandsInfo) VULKAN_LAYER_BODY
+	static void CmdPreprocessGeneratedCommandsEXT(VkResult Result, VkCommandBuffer CommandBuffer, const VkGeneratedCommandsInfoEXT* GeneratedCommandsInfo, VkCommandBuffer StateCommandBuffer) VULKAN_LAYER_BODY
+	static void CreateIndirectCommandsLayoutEXT(VkResult Result, VkDevice Device, const VkIndirectCommandsLayoutCreateInfoEXT* CreateInfo, const VkAllocationCallbacks* Allocator, VkIndirectCommandsLayoutEXT* IndirectCommandsLayout) VULKAN_LAYER_BODY
+	static void CreateIndirectExecutionSetEXT(VkResult Result, VkDevice Device, const VkIndirectExecutionSetCreateInfoEXT* CreateInfo, const VkAllocationCallbacks* Allocator, VkIndirectExecutionSetEXT* IndirectExecutionSet) VULKAN_LAYER_BODY
+	static void DestroyIndirectCommandsLayoutEXT(VkResult Result, VkDevice Device, VkIndirectCommandsLayoutEXT IndirectCommandsLayout, const VkAllocationCallbacks* Allocator) VULKAN_LAYER_BODY
+	static void DestroyIndirectExecutionSetEXT(VkResult Result, VkDevice Device, VkIndirectExecutionSetEXT IndirectExecutionSet, const VkAllocationCallbacks* Allocator) VULKAN_LAYER_BODY
+	static void GetGeneratedCommandsMemoryRequirementsEXT(VkResult Result, VkDevice Device, const VkGeneratedCommandsMemoryRequirementsInfoEXT* Info, VkMemoryRequirements2* MemoryRequirements) VULKAN_LAYER_BODY
+	static void UpdateIndirectExecutionSetPipelineEXT(VkResult Result, VkDevice Device, VkIndirectExecutionSetEXT IndirectExecutionSet, uint32_t ExecutionSetWriteCount, const VkWriteIndirectExecutionSetPipelineEXT* ExecutionSetWrites) VULKAN_LAYER_BODY
+	static void UpdateIndirectExecutionSetShaderEXT(VkResult Result, VkDevice Device, VkIndirectExecutionSetEXT IndirectExecutionSet, uint32_t ExecutionSetWriteCount, const VkWriteIndirectExecutionSetShaderEXT* ExecutionSetWrites) VULKAN_LAYER_BODY
+
 	static void CopyMemoryToImageEXT(VkResult Result, VkDevice Device, const VkCopyMemoryToImageInfo* CopyMemoryToImageInfo) VULKAN_LAYER_BODY
 	static void TransitionImageLayoutEXT(VkResult Result, VkDevice Device, uint32_t TransitionCount, const VkHostImageLayoutTransitionInfo* Transitions) VULKAN_LAYER_BODY
 };
@@ -2119,6 +2129,72 @@ namespace VulkanRHI
 		FWrapLayer::CmdSetFragmentShadingRateKHR(VK_SUCCESS, CommandBuffer, FragmentSize, combinerOps);
 	}
 
+	static FORCEINLINE_DEBUGGABLE void vkCmdExecuteGeneratedCommandsEXT(VkCommandBuffer CommandBuffer, VkBool32 IsPreprocessed, const VkGeneratedCommandsInfoEXT* GeneratedCommandsInfo)
+	{
+		FWrapLayer::CmdExecuteGeneratedCommandsEXT(VK_RESULT_MAX_ENUM, CommandBuffer, IsPreprocessed, GeneratedCommandsInfo);
+		VULKANAPINAMESPACE::vkCmdExecuteGeneratedCommandsEXT(CommandBuffer, IsPreprocessed, GeneratedCommandsInfo);
+		FWrapLayer::CmdExecuteGeneratedCommandsEXT(VK_SUCCESS, CommandBuffer, IsPreprocessed, GeneratedCommandsInfo);
+	}
+
+	static FORCEINLINE_DEBUGGABLE void vkCmdPreprocessGeneratedCommandsEXT(VkCommandBuffer CommandBuffer, const VkGeneratedCommandsInfoEXT* GeneratedCommandsInfo, VkCommandBuffer StateCommandBuffer)
+	{
+		FWrapLayer::CmdPreprocessGeneratedCommandsEXT(VK_RESULT_MAX_ENUM, CommandBuffer, GeneratedCommandsInfo, StateCommandBuffer);
+		VULKANAPINAMESPACE::vkCmdPreprocessGeneratedCommandsEXT(CommandBuffer, GeneratedCommandsInfo, StateCommandBuffer);
+		FWrapLayer::CmdPreprocessGeneratedCommandsEXT(VK_SUCCESS, CommandBuffer, GeneratedCommandsInfo, StateCommandBuffer);
+	}
+
+	static FORCEINLINE_DEBUGGABLE VkResult vkCreateIndirectCommandsLayoutEXT(VkDevice Device, const VkIndirectCommandsLayoutCreateInfoEXT* CreateInfo, const VkAllocationCallbacks* Allocator, VkIndirectCommandsLayoutEXT* IndirectCommandsLayout)
+	{
+		FWrapLayer::CreateIndirectCommandsLayoutEXT(VK_RESULT_MAX_ENUM, Device, CreateInfo, Allocator, IndirectCommandsLayout);
+		VkResult Result = VULKANAPINAMESPACE::vkCreateIndirectCommandsLayoutEXT(Device, CreateInfo, Allocator, IndirectCommandsLayout);
+		FWrapLayer::CreateIndirectCommandsLayoutEXT(VK_SUCCESS, Device, CreateInfo, Allocator, IndirectCommandsLayout);
+		return Result;
+	}
+
+	static FORCEINLINE_DEBUGGABLE VkResult vkCreateIndirectExecutionSetEXT(VkDevice Device, const VkIndirectExecutionSetCreateInfoEXT* CreateInfo, const VkAllocationCallbacks* Allocator, VkIndirectExecutionSetEXT* IndirectExecutionSet)
+	{
+		FWrapLayer::CreateIndirectExecutionSetEXT(VK_RESULT_MAX_ENUM, Device, CreateInfo, Allocator, IndirectExecutionSet);
+		VkResult Result = VULKANAPINAMESPACE::vkCreateIndirectExecutionSetEXT(Device, CreateInfo, Allocator, IndirectExecutionSet);
+		FWrapLayer::CreateIndirectExecutionSetEXT(VK_SUCCESS, Device, CreateInfo, Allocator, IndirectExecutionSet);
+		return Result;
+	}
+
+	static FORCEINLINE_DEBUGGABLE void vkDestroyIndirectCommandsLayoutEXT(VkDevice Device, VkIndirectCommandsLayoutEXT IndirectCommandsLayout, const VkAllocationCallbacks* Allocator)
+	{
+		FWrapLayer::DestroyIndirectCommandsLayoutEXT(VK_RESULT_MAX_ENUM, Device, IndirectCommandsLayout, Allocator);
+		VULKANAPINAMESPACE::vkDestroyIndirectCommandsLayoutEXT(Device, IndirectCommandsLayout, Allocator);
+		FWrapLayer::DestroyIndirectCommandsLayoutEXT(VK_SUCCESS, Device, IndirectCommandsLayout, Allocator);
+	}
+
+	static FORCEINLINE_DEBUGGABLE void vkDestroyIndirectExecutionSetEXT(VkDevice Device, VkIndirectExecutionSetEXT IndirectExecutionSet, const VkAllocationCallbacks* Allocator)
+	{
+		FWrapLayer::DestroyIndirectExecutionSetEXT(VK_RESULT_MAX_ENUM, Device, IndirectExecutionSet, Allocator);
+		VULKANAPINAMESPACE::vkDestroyIndirectExecutionSetEXT(Device, IndirectExecutionSet, Allocator);
+		FWrapLayer::DestroyIndirectExecutionSetEXT(VK_SUCCESS, Device, IndirectExecutionSet, Allocator);
+	}
+
+	static FORCEINLINE_DEBUGGABLE void vkGetGeneratedCommandsMemoryRequirementsEXT(VkDevice Device, const VkGeneratedCommandsMemoryRequirementsInfoEXT* Info, VkMemoryRequirements2* MemoryRequirements)
+	{
+		FWrapLayer::GetGeneratedCommandsMemoryRequirementsEXT(VK_RESULT_MAX_ENUM, Device, Info, MemoryRequirements);
+		VULKANAPINAMESPACE::vkGetGeneratedCommandsMemoryRequirementsEXT(Device, Info, MemoryRequirements);
+		FWrapLayer::GetGeneratedCommandsMemoryRequirementsEXT(VK_SUCCESS, Device, Info, MemoryRequirements);
+	}
+
+	static FORCEINLINE_DEBUGGABLE void vkUpdateIndirectExecutionSetPipelineEXT(VkDevice Device, VkIndirectExecutionSetEXT IndirectExecutionSet, uint32_t ExecutionSetWriteCount, const VkWriteIndirectExecutionSetPipelineEXT* ExecutionSetWrites)
+	{
+		FWrapLayer::UpdateIndirectExecutionSetPipelineEXT(VK_RESULT_MAX_ENUM, Device, IndirectExecutionSet, ExecutionSetWriteCount, ExecutionSetWrites);
+		VULKANAPINAMESPACE::vkUpdateIndirectExecutionSetPipelineEXT(Device, IndirectExecutionSet, ExecutionSetWriteCount, ExecutionSetWrites);
+		FWrapLayer::UpdateIndirectExecutionSetPipelineEXT(VK_SUCCESS, Device, IndirectExecutionSet, ExecutionSetWriteCount, ExecutionSetWrites);
+	}
+
+	static FORCEINLINE_DEBUGGABLE void vkUpdateIndirectExecutionSetShaderEXT(VkDevice Device, VkIndirectExecutionSetEXT IndirectExecutionSet, uint32_t ExecutionSetWriteCount, const VkWriteIndirectExecutionSetShaderEXT* ExecutionSetWrites)
+	{
+		FWrapLayer::UpdateIndirectExecutionSetShaderEXT(VK_RESULT_MAX_ENUM, Device, IndirectExecutionSet, ExecutionSetWriteCount, ExecutionSetWrites);
+		VULKANAPINAMESPACE::vkUpdateIndirectExecutionSetShaderEXT(Device, IndirectExecutionSet, ExecutionSetWriteCount, ExecutionSetWrites);
+		FWrapLayer::UpdateIndirectExecutionSetShaderEXT(VK_SUCCESS, Device, IndirectExecutionSet, ExecutionSetWriteCount, ExecutionSetWrites);
+	}
+
+
 	static FORCEINLINE_DEBUGGABLE VkResult vkCopyMemoryToImageEXT(VkDevice Device, const VkCopyMemoryToImageInfo* CopyMemoryToImageInfo)
 	{
 		FWrapLayer::CopyMemoryToImageEXT(VK_RESULT_MAX_ENUM, Device, CopyMemoryToImageInfo);
diff --git a/Engine/Source/Runtime/VulkanRHI/Private/VulkanDebug.cpp b/Engine/Source/Runtime/VulkanRHI/Private/VulkanDebug.cpp
--- a/Engine/Source/Runtime/VulkanRHI/Private/VulkanDebug.cpp
+++ b/Engine/Source/Runtime/VulkanRHI/Private/VulkanDebug.cpp
@@ -5046,6 +5046,98 @@ void FWrapLayer::CmdSetFragmentShadingRateKHR(VkResult Result, VkCommandBuffer C
 	}
 }
 
+void FWrapLayer::CmdExecuteGeneratedCommandsEXT(VkResult Result, VkCommandBuffer CommandBuffer, VkBool32 IsPreprocessed, const VkGeneratedCommandsInfoEXT* GeneratedCommandsInfo)
+{
+	if (Result == VK_RESULT_MAX_ENUM)
+	{
+#if VULKAN_ENABLE_DUMP_LAYER
+		PrintfBeginResult(FString::Printf(TEXT("CmdExecuteGeneratedCommandsEXT()")));
+#endif
+	}
+}
+
+void FWrapLayer::CmdPreprocessGeneratedCommandsEXT(VkResult Result, VkCommandBuffer CommandBuffer, const VkGeneratedCommandsInfoEXT* GeneratedCommandsInfo, VkCommandBuffer StateCommandBuffer)
+{
+	if (Result == VK_RESULT_MAX_ENUM)
+	{
+#if VULKAN_ENABLE_DUMP_LAYER
+		PrintfBeginResult(FString::Printf(TEXT("CmdPreprocessGeneratedCommandsEXT()")));
+#endif
+	}
+}
+
+void FWrapLayer::CreateIndirectCommandsLayoutEXT(VkResult Result, VkDevice Device, const VkIndirectCommandsLayoutCreateInfoEXT* CreateInfo, const VkAllocationCallbacks* Allocator, VkIndirectCommandsLayoutEXT* IndirectCommandsLayout)
+{
+	if (Result == VK_RESULT_MAX_ENUM)
+	{
+#if VULKAN_ENABLE_DUMP_LAYER
+		PrintfBeginResult(FString::Printf(TEXT("CreateIndirectCommandsLayoutEXT()")));
+#endif
+	}
+}
+
+void FWrapLayer::CreateIndirectExecutionSetEXT(VkResult Result, VkDevice Device, const VkIndirectExecutionSetCreateInfoEXT* CreateInfo, const VkAllocationCallbacks* Allocator, VkIndirectExecutionSetEXT* IndirectExecutionSet)
+{
+	if (Result == VK_RESULT_MAX_ENUM)
+	{
+#if VULKAN_ENABLE_DUMP_LAYER
+		PrintfBeginResult(FString::Printf(TEXT("CreateIndirectExecutionSetEXT()")));
+#endif
+	}
+}
+
+void FWrapLayer::DestroyIndirectCommandsLayoutEXT(VkResult Result, VkDevice Device, VkIndirectCommandsLayoutEXT IndirectCommandsLayout, const VkAllocationCallbacks* Allocator)
+{
+	if (Result == VK_RESULT_MAX_ENUM)
+	{
+#if VULKAN_ENABLE_DUMP_LAYER
+		PrintfBeginResult(FString::Printf(TEXT("DestroyIndirectCommandsLayoutEXT()")));
+#endif
+	}
+}
+
+void FWrapLayer::DestroyIndirectExecutionSetEXT(VkResult Result, VkDevice Device, VkIndirectExecutionSetEXT IndirectExecutionSet, const VkAllocationCallbacks* Allocator)
+{
+	if (Result == VK_RESULT_MAX_ENUM)
+	{
+#if VULKAN_ENABLE_DUMP_LAYER
+		PrintfBeginResult(FString::Printf(TEXT("DestroyIndirectExecutionSetEXT()")));
+#endif
+	}
+}
+
+void FWrapLayer::GetGeneratedCommandsMemoryRequirementsEXT(VkResult Result, VkDevice Device, const VkGeneratedCommandsMemoryRequirementsInfoEXT* Info, VkMemoryRequirements2* MemoryRequirements)
+{
+	if (Result == VK_RESULT_MAX_ENUM)
+	{
+#if VULKAN_ENABLE_DUMP_LAYER
+		PrintfBeginResult(FString::Printf(TEXT("GetGeneratedCommandsMemoryRequirementsEXT()")));
+#endif
+	}
+}
+
+void FWrapLayer::UpdateIndirectExecutionSetPipelineEXT(VkResult Result, VkDevice Device, VkIndirectExecutionSetEXT IndirectExecutionSet, uint32_t ExecutionSetWriteCount, const VkWriteIndirectExecutionSetPipelineEXT* ExecutionSetWrites)
+{
+	if (Result == VK_RESULT_MAX_ENUM)
+	{
+#if VULKAN_ENABLE_DUMP_LAYER
+		PrintfBeginResult(FString::Printf(TEXT("UpdateIndirectExecutionSetPipelineEXT()")));
+#endif
+	}
+}
+
+void FWrapLayer::UpdateIndirectExecutionSetShaderEXT(VkResult Result, VkDevice Device, VkIndirectExecutionSetEXT IndirectExecutionSet, uint32_t ExecutionSetWriteCount, const VkWriteIndirectExecutionSetShaderEXT* ExecutionSetWrites)
+{
+	if (Result == VK_RESULT_MAX_ENUM)
+	{
+#if VULKAN_ENABLE_DUMP_LAYER
+		PrintfBeginResult(FString::Printf(TEXT("UpdateIndirectExecutionSetShaderEXT()")));
+#endif
+	}
+}
+
+
+
 void FWrapLayer::CopyMemoryToImageEXT(VkResult Result, VkDevice Device, const VkCopyMemoryToImageInfo* CopyMemoryToImageInfo)
 {
 	if (Result == VK_RESULT_MAX_ENUM)
diff --git a/Engine/Source/Runtime/VulkanRHI/Private/VulkanDevice.h b/Engine/Source/Runtime/VulkanRHI/Private/VulkanDevice.h
--- a/Engine/Source/Runtime/VulkanRHI/Private/VulkanDevice.h
+++ b/Engine/Source/Runtime/VulkanRHI/Private/VulkanDevice.h
@@ -177,6 +177,8 @@ struct FOptionalVulkanDeviceExtensionProperties
 	VkPhysicalDeviceGraphicsPipelineLibraryPropertiesEXT GraphicsPipelineLibraryProperties;
 
 	VkPhysicalDeviceMeshShaderPropertiesEXT MeshShaderProperties;
+
+	VkPhysicalDeviceDeviceGeneratedCommandsPropertiesEXT DeviceGeneratedCommandsProperties;
 };
 
 class FVulkanPhysicalDeviceFeatures
diff --git a/Engine/Source/Runtime/VulkanRHI/Private/VulkanExtensions.cpp b/Engine/Source/Runtime/VulkanRHI/Private/VulkanExtensions.cpp
--- a/Engine/Source/Runtime/VulkanRHI/Private/VulkanExtensions.cpp
+++ b/Engine/Source/Runtime/VulkanRHI/Private/VulkanExtensions.cpp
@@ -2283,6 +2283,7 @@ public:
 
 	FVulkanEXTDeviceGeneratedCommandsExtension(FVulkanDevice* InDevice)
 		: FVulkanDeviceExtension(InDevice, VK_EXT_DEVICE_GENERATED_COMMANDS_EXTENSION_NAME, VULKAN_EXTENSION_ENABLED, VK_API_VERSION_1_4)
+		, DeviceGeneratedCommandsProperties(GetDeviceExtensionProperties().DeviceGeneratedCommandsProperties)
 	{
 		bEnabledInCode = bEnabledInCode && FDataDrivenShaderPlatformInfo::GetSupportsShaderBundleDispatch(GMaxRHIShaderPlatform);
 	}
@@ -2330,7 +2331,7 @@ public:
 	}
 
 	VkPhysicalDeviceDeviceGeneratedCommandsFeaturesEXT DeviceGeneratedCommandsFeatures;
-	VkPhysicalDeviceDeviceGeneratedCommandsPropertiesEXT DeviceGeneratedCommandsProperties;
+	VkPhysicalDeviceDeviceGeneratedCommandsPropertiesEXT& DeviceGeneratedCommandsProperties;
 };
 
 
@@ -2369,9 +2370,9 @@ public:
 	{
 		if (bRequirementsPassed)
 		{
-			UE_LOG(LogVulkanRHI, Display, TEXT("Robustness2 Enabled: Buffer=%d (Align SB/UB=%d/%d), Image=%d, NullDescriptor=%d"),
+			UE_LOG(LogVulkanRHI, Display, TEXT("Robustness2 Enabled: Buffer=%d (Align SB/UB=%u/%u), Image=%d, NullDescriptor=%d"),
 				Robustness2Features.robustBufferAccess2, 
-				Robustness2Properties.robustStorageBufferAccessSizeAlignment, Robustness2Properties.robustUniformBufferAccessSizeAlignment,
+				(uint32)Robustness2Properties.robustStorageBufferAccessSizeAlignment, (uint32)Robustness2Properties.robustUniformBufferAccessSizeAlignment,
 				Robustness2Features.robustImageAccess2, Robustness2Features.nullDescriptor);
 		}
 	}
diff --git a/Engine/Source/Runtime/VulkanRHI/Private/VulkanLoader.h b/Engine/Source/Runtime/VulkanRHI/Private/VulkanLoader.h
--- a/Engine/Source/Runtime/VulkanRHI/Private/VulkanLoader.h
+++ b/Engine/Source/Runtime/VulkanRHI/Private/VulkanLoader.h
@@ -277,10 +277,16 @@
 	EnumMacro(PFN_vkCopyMemoryToImageEXT, vkCopyMemoryToImageEXT) \
 	EnumMacro(PFN_vkTransitionImageLayoutEXT, vkTransitionImageLayoutEXT) \
 	EnumMacro(PFN_vkCmdPushDescriptorSetKHR, vkCmdPushDescriptorSetKHR) \
-	EnumMacro(PFN_vkCmdPushDescriptorSet2KHR, vkCmdPushDescriptorSet2KHR)
-
-
-
+	EnumMacro(PFN_vkCmdPushDescriptorSet2KHR, vkCmdPushDescriptorSet2KHR) \
+	EnumMacro(PFN_vkCmdExecuteGeneratedCommandsEXT, vkCmdExecuteGeneratedCommandsEXT) \
+	EnumMacro(PFN_vkCmdPreprocessGeneratedCommandsEXT, vkCmdPreprocessGeneratedCommandsEXT) \
+	EnumMacro(PFN_vkCreateIndirectCommandsLayoutEXT, vkCreateIndirectCommandsLayoutEXT) \
+	EnumMacro(PFN_vkCreateIndirectExecutionSetEXT, vkCreateIndirectExecutionSetEXT) \
+	EnumMacro(PFN_vkDestroyIndirectCommandsLayoutEXT, vkDestroyIndirectCommandsLayoutEXT) \
+	EnumMacro(PFN_vkDestroyIndirectExecutionSetEXT, vkDestroyIndirectExecutionSetEXT) \
+	EnumMacro(PFN_vkGetGeneratedCommandsMemoryRequirementsEXT, vkGetGeneratedCommandsMemoryRequirementsEXT) \
+	EnumMacro(PFN_vkUpdateIndirectExecutionSetPipelineEXT, vkUpdateIndirectExecutionSetPipelineEXT) \
+	EnumMacro(PFN_vkUpdateIndirectExecutionSetShaderEXT, vkUpdateIndirectExecutionSetShaderEXT)
 
 
 
