VulkanRHI: Keep an array of push contant ranges and identify each range with an enum for the system they belong to.

#rnx

diff --git a/Engine/Source/Runtime/VulkanRHI/Private/VulkanBindlessDescriptorManager.cpp b/Engine/Source/Runtime/VulkanRHI/Private/VulkanBindlessDescriptorManager.cpp
--- a/Engine/Source/Runtime/VulkanRHI/Private/VulkanBindlessDescriptorManager.cpp
+++ b/Engine/Source/Runtime/VulkanRHI/Private/VulkanBindlessDescriptorManager.cpp
@@ -209,6 +209,19 @@ static inline uint32 GetInitialDescriptorCount(VkDescriptorType DescriptorType)
 }
 
 
+const VkPushConstantRange FVulkanBindlessDescriptorManager::PushConstantRanges[(uint8)FVulkanBindlessDescriptorManager::EPushConstantRangeType::Count] =
+{
+	// StaticUniformBuffers
+	{
+		.stageFlags = VK_SHADER_STAGE_ALL,
+		.offset = 0,
+		.size = (FRHIShaderBindingLayout::MaxUniformBufferEntries * sizeof(uint32))
+	}
+
+	// NOTE: Vulkan Core has a 128byte minimum (256 for Vulkan 1.4)
+};
+
+
 static inline VkMemoryPropertyFlags GetDescriptorBufferMemoryType(FVulkanDevice& Device)
 {
 	if (Device.HasUnifiedMemory() || (FVulkanPlatform::SupportsDeviceLocalHostVisibleWithNoPenalty(Device.GetVendorId()) &&
@@ -630,19 +643,12 @@ void FVulkanBindlessDescriptorManager::Init()
 		ZeroVulkanStruct(PipelineLayoutCreateInfo, VK_STRUCTURE_TYPE_PIPELINE_LAYOUT_CREATE_INFO);
 		PipelineLayoutCreateInfo.setLayoutCount = BindlessDescriptorSetLayouts.Num();
 		PipelineLayoutCreateInfo.pSetLayouts = BindlessDescriptorSetLayouts.GetData();
+		PipelineLayoutCreateInfo.pushConstantRangeCount = UE_ARRAY_COUNT(PushConstantRanges);
+		PipelineLayoutCreateInfo.pPushConstantRanges = PushConstantRanges;
 
-		// Set the static uniform buffers with push constants (only enabled for ray tracing in DDPI, only used for RGS)
-		VkPushConstantRange PushConstants;
-		if (RHIGetStaticShaderBindingLayoutSupport(GMaxRHIShaderPlatform) != ERHIStaticShaderBindingLayoutSupport::Unsupported)
-		{
-			checkf(RHIGetStaticShaderBindingLayoutSupport(GMaxRHIShaderPlatform) == ERHIStaticShaderBindingLayoutSupport::RayTracingOnly,
-				TEXT("StaticShaderBindingLayoutSupport is currently only supported for ray gen shaders in Vulkan."));
-			PushConstants.offset = 0;
-			PushConstants.size = FRHIShaderBindingLayout::MaxUniformBufferEntries * sizeof(uint32);
-			PushConstants.stageFlags = VK_SHADER_STAGE_ALL;
-			PipelineLayoutCreateInfo.pushConstantRangeCount = 1;
-			PipelineLayoutCreateInfo.pPushConstantRanges = &PushConstants;
-		}
+		checkf((uint32)EPushConstantRangeType::Count == UE_ARRAY_COUNT(PushConstantRanges),
+			TEXT("Mismatch between expected push constant range count (%d) and provided array size (%d)"), 
+			(uint32)EPushConstantRangeType::Count, UE_ARRAY_COUNT(PushConstantRanges));
 
 		VERIFYVULKANRESULT(VulkanRHI::vkCreatePipelineLayout(DeviceHandle, &PipelineLayoutCreateInfo, VULKAN_CPU_ALLOCATOR, &BindlessPipelineLayout));
 		VULKAN_SET_DEBUG_NAME(Device, VK_OBJECT_TYPE_PIPELINE_LAYOUT, BindlessPipelineLayout, TEXT("BindlessPipelineLayout(SetCount=%d)"), VulkanBindless::MaxNumSets);
diff --git a/Engine/Source/Runtime/VulkanRHI/Private/VulkanBindlessDescriptorManager.h b/Engine/Source/Runtime/VulkanRHI/Private/VulkanBindlessDescriptorManager.h
--- a/Engine/Source/Runtime/VulkanRHI/Private/VulkanBindlessDescriptorManager.h
+++ b/Engine/Source/Runtime/VulkanRHI/Private/VulkanBindlessDescriptorManager.h
@@ -16,6 +16,14 @@ class FVulkanBuffer;
 class FVulkanBindlessDescriptorManager
 {
 public:
+
+	enum class EPushConstantRangeType : uint8
+	{
+		StaticUniformBuffers,
+
+		Count
+	};
+
 	FVulkanBindlessDescriptorManager(FVulkanDevice& InDevice);
 	~FVulkanBindlessDescriptorManager();
 
@@ -38,6 +46,11 @@ public:
 		return BindlessDescriptorSetLayouts;
 	}
 
+	static TConstArrayView<VkPushConstantRange> GetPushConstantRanges()
+	{
+		return PushConstantRanges;
+	}
+
 	void BindDescriptorBuffers(VkCommandBuffer CommandBuffer, VkPipelineStageFlags SupportedStages);
 
 	FRHIDescriptorHandle AllocateDescriptor(VkDescriptorType DescriptorType);
@@ -98,6 +111,7 @@ private:
 
 	VkPipelineLayout BindlessPipelineLayout = VK_NULL_HANDLE;
 	TStaticArray<VkDescriptorSetLayout, VulkanBindless::MaxNumSets> BindlessDescriptorSetLayouts;
+	static const VkPushConstantRange PushConstantRanges[(uint8)EPushConstantRangeType::Count];
 
 	BindlessSetState& GetBindlessState(ERHIDescriptorType DescriptorType);
 	BindlessSetState& GetBindlessState(VkDescriptorType DescriptorType);
diff --git a/Engine/Source/Runtime/VulkanRHI/Private/VulkanShaderObjectManager.cpp b/Engine/Source/Runtime/VulkanRHI/Private/VulkanShaderObjectManager.cpp
--- a/Engine/Source/Runtime/VulkanRHI/Private/VulkanShaderObjectManager.cpp
+++ b/Engine/Source/Runtime/VulkanRHI/Private/VulkanShaderObjectManager.cpp
@@ -47,6 +47,7 @@ VkShaderEXT FVulkanShaderObjectManager::CreateShaderObject(FVulkanShader& Vulkan
 	VulkanShader.GetEntryPoint(EntryPoint, 24);
 
 	TConstArrayView<VkDescriptorSetLayout> DescriptorSetLayout = Device.GetBindlessDescriptorManager()->GetDescriptorSetLayouts();
+	TConstArrayView<VkPushConstantRange> PushConstantRanges = Device.GetBindlessDescriptorManager()->GetPushConstantRanges();
 
 	VkShaderCreateInfoEXT ShaderCreateInfo = {
 		.sType = VK_STRUCTURE_TYPE_SHADER_CREATE_INFO_EXT,
@@ -60,23 +61,11 @@ VkShaderEXT FVulkanShaderObjectManager::CreateShaderObject(FVulkanShader& Vulkan
 		.pName = EntryPoint,
 		.setLayoutCount = (uint32)DescriptorSetLayout.Num(),
 		.pSetLayouts = DescriptorSetLayout.GetData(),
-		.pushConstantRangeCount = 0,
-		.pPushConstantRanges = nullptr,
+		.pushConstantRangeCount = (uint32)PushConstantRanges.Num(),
+		.pPushConstantRanges = PushConstantRanges.GetData(),
 		.pSpecializationInfo = nullptr
 	};
 
-	VkPushConstantRange PushConstants;
-	if (RHIGetStaticShaderBindingLayoutSupport(GMaxRHIShaderPlatform) != ERHIStaticShaderBindingLayoutSupport::Unsupported)
-	{
-		checkf(RHIGetStaticShaderBindingLayoutSupport(GMaxRHIShaderPlatform) == ERHIStaticShaderBindingLayoutSupport::RayTracingOnly,
-			TEXT("StaticShaderBindingLayoutSupport is currently only supported for ray gen shaders in Vulkan."));
-		PushConstants.offset = 0;
-		PushConstants.size = FRHIShaderBindingLayout::MaxUniformBufferEntries * sizeof(uint32);
-		PushConstants.stageFlags = VK_SHADER_STAGE_ALL;
-		ShaderCreateInfo.pushConstantRangeCount = 1;
-		ShaderCreateInfo.pPushConstantRanges = &PushConstants;
-	}
-
 	// :todo-jn: Task shaders never used at the moment
 	if (Stage == VK_SHADER_STAGE_MESH_BIT_EXT)
 	{
