Cleanup / remove more uses of the immediate context in VulkanRHI
- Removed FVulkanCommandListContextImmediate. The type was unnecessary.
- Added bImmediate bool to contexts. The "Immediate" pointer was only used in boolean conditions, never dereferenced.

Removed RHISetImageLayout
- Nothing was using this function, and it is dangerous (directly uses the immediate context).

Deprecated RHIGetActiveVkCommandBuffer and replaced implementation with a checkNoEntry()
- Again, the function is dangerous. Correct use is to obtain the Vulkan RHI context from an executing RHICmdList, via RHIGetVkCommandBuffer.

Removed some old deprecated functions

diff --git a/Engine/Source/Runtime/VulkanRHI/Private/VulkanCommands.cpp b/Engine/Source/Runtime/VulkanRHI/Private/VulkanCommands.cpp
--- a/Engine/Source/Runtime/VulkanRHI/Private/VulkanCommands.cpp
+++ b/Engine/Source/Runtime/VulkanRHI/Private/VulkanCommands.cpp
@@ -904,5 +904,5 @@ void FVulkanCommandListContext::RHIWriteGPUFence(FRHIGPUFence* FenceRHI)
 IRHIComputeContext* FVulkanDynamicRHI::RHIGetCommandContext(ERHIPipeline Pipeline)
 {
 	check(IsSingleRHIPipeline(Pipeline));
-	return new FVulkanCommandListContext(*Device, Pipeline, &Device->GetImmediateContext());
+	return new FVulkanCommandListContext(*Device, Pipeline, false);
 }
diff --git a/Engine/Source/Runtime/VulkanRHI/Private/VulkanContext.cpp b/Engine/Source/Runtime/VulkanRHI/Private/VulkanContext.cpp
--- a/Engine/Source/Runtime/VulkanRHI/Private/VulkanContext.cpp
+++ b/Engine/Source/Runtime/VulkanRHI/Private/VulkanContext.cpp
@@ -253,11 +253,11 @@ void FVulkanContextCommon::HandleReservedResourceCommits(TArrayView<const FRHITr
 
 
 
-FVulkanCommandListContext::FVulkanCommandListContext(FVulkanDevice& InDevice, ERHIPipeline InPipeline, FVulkanCommandListContext* InImmediate)
+FVulkanCommandListContext::FVulkanCommandListContext(FVulkanDevice& InDevice, ERHIPipeline InPipeline, bool bInImmediate)
 	: FVulkanContextCommon(InDevice, *InDevice.GetQueue(InPipeline), EVulkanCommandBufferType::Primary)
-	, Immediate(InImmediate)
 	, RHIPipeline(InPipeline)
 	, bSupportsBreadcrumbs(GRHIGlobals.SupportsTimestampRenderQueries)
+	, bImmediate(bInImmediate)
 #if (RHI_NEW_GPU_PROFILER == 0)
 	, GpuProfiler(this, &InDevice)
 #endif
@@ -274,12 +274,11 @@ FVulkanCommandListContext::FVulkanCommandListContext(FVulkanDevice& InDevice, ER
 	GlobalUniformBuffers.AddZeroed(FUniformBufferStaticSlotRegistry::Get().GetSlotCount());
 }
 
-FVulkanCommandListContext::FVulkanCommandListContext(FVulkanDevice& InDevice, FVulkanCommandListContext* InImmediate, FVulkanParallelRenderPassInfo* InParallelRenderPassInfo)
-	: FVulkanContextCommon(InDevice, *InDevice.GetQueue(ERHIPipeline::Graphics), 
-		InDevice.GetOptionalExtensions().HasKHRDynamicRendering ? EVulkanCommandBufferType::Parallel : EVulkanCommandBufferType::Secondary)
-	, Immediate(InImmediate)
+FVulkanCommandListContext::FVulkanCommandListContext(FVulkanDevice& InDevice, FVulkanParallelRenderPassInfo* InParallelRenderPassInfo)
+	: FVulkanContextCommon(InDevice, *InDevice.GetQueue(ERHIPipeline::Graphics), InDevice.GetOptionalExtensions().HasKHRDynamicRendering ? EVulkanCommandBufferType::Parallel : EVulkanCommandBufferType::Secondary)
 	, RHIPipeline(ERHIPipeline::Graphics)
 	, bSupportsBreadcrumbs(false)
+	, bImmediate(false)
 	, CurrentParallelRenderPassInfo(InParallelRenderPassInfo)
 #if (RHI_NEW_GPU_PROFILER == 0)
 	, GpuProfiler(this, &InDevice)
diff --git a/Engine/Source/Runtime/VulkanRHI/Private/VulkanContext.h b/Engine/Source/Runtime/VulkanRHI/Private/VulkanContext.h
--- a/Engine/Source/Runtime/VulkanRHI/Private/VulkanContext.h
+++ b/Engine/Source/Runtime/VulkanRHI/Private/VulkanContext.h
@@ -240,11 +240,11 @@ private:
 class FVulkanCommandListContext : public FVulkanContextCommon, public IRHICommandContext
 {
 public:
-	FVulkanCommandListContext(FVulkanDevice& InDevice, ERHIPipeline InPipeline, FVulkanCommandListContext* InImmediate);
+	FVulkanCommandListContext(FVulkanDevice& InDevice, ERHIPipeline InPipeline, bool bInImmediate);
 	virtual ~FVulkanCommandListContext();
 
 	// Constructor for parallel render contexts that use secondary command buffers
-	FVulkanCommandListContext(FVulkanDevice& InDevice, FVulkanCommandListContext* InImmediate, FVulkanParallelRenderPassInfo* InParallelRenderPassInfo);
+	FVulkanCommandListContext(FVulkanDevice& InDevice, FVulkanParallelRenderPassInfo* InParallelRenderPassInfo);
 
 	static inline FVulkanCommandListContext& Get(FRHICommandListBase& RHICmdList)
 	{
@@ -253,7 +253,7 @@ public:
 
 	inline bool IsImmediate() const
 	{
-		return Immediate == nullptr;
+		return bImmediate;
 	}
 
 	virtual ERHIPipeline GetPipeline() const override
@@ -426,9 +426,9 @@ protected:
 	void ExtractDepthStencilLayouts(const FRHIRenderPassInfo& InInfo, VkImageLayout& OutDepthLayout, VkImageLayout& OutStencilLayout);
 	void FillDynamicRenderingInfo(const FRHIRenderPassInfo& RenderPassInfo, FVulkanDynamicRenderingInfo& OutRenderingInfo);
 
-	FVulkanCommandListContext* Immediate;
 	const ERHIPipeline RHIPipeline;
 	const bool bSupportsBreadcrumbs;
+	const bool bImmediate;
 
 	void BeginOcclusionQueryBatch(uint32 NumQueriesInBatch, uint32 MultiViewCount);
 
@@ -478,17 +478,6 @@ private:
 	const FRHIShaderBindingLayout* ShaderBindingLayout = nullptr;
 };
 
-class FVulkanCommandListContextImmediate : public FVulkanCommandListContext
-{
-public:
-	static inline FVulkanCommandListContextImmediate& Get(FRHICommandListBase& RHICmdList)
-	{
-		return static_cast<FVulkanCommandListContextImmediate&>(RHICmdList.GetContext().GetLowestLevelContext());
-	}
-
-	FVulkanCommandListContextImmediate(FVulkanDevice& InDevice);
-};
-
 class FVulkanUploadContext : public FVulkanContextCommon, public IRHIUploadContext
 {
 public:
diff --git a/Engine/Source/Runtime/VulkanRHI/Private/VulkanDevice.cpp b/Engine/Source/Runtime/VulkanRHI/Private/VulkanDevice.cpp
--- a/Engine/Source/Runtime/VulkanRHI/Private/VulkanDevice.cpp
+++ b/Engine/Source/Runtime/VulkanRHI/Private/VulkanDevice.cpp
@@ -1422,7 +1422,7 @@ void FVulkanDevice::InitGPU()
 	// always look in the saved directory (for the cache from previous run that wasn't moved over to stage directory)
 	CacheFilenames.Add(VulkanRHI::GetPipelineCacheFilename());
 
-	ImmediateContext = new FVulkanCommandListContextImmediate(*this);
+	ImmediateContext = new FVulkanCommandListContext(*this, ERHIPipeline::Graphics, true);
 
 #if VULKAN_SUPPORTS_VALIDATION_CACHE
 	if (OptionalDeviceExtensions.HasEXTValidationCache)
diff --git a/Engine/Source/Runtime/VulkanRHI/Private/VulkanDevice.h b/Engine/Source/Runtime/VulkanRHI/Private/VulkanDevice.h
--- a/Engine/Source/Runtime/VulkanRHI/Private/VulkanDevice.h
+++ b/Engine/Source/Runtime/VulkanRHI/Private/VulkanDevice.h
@@ -19,7 +19,6 @@ class FVulkanDescriptorSetCache;
 class FVulkanDescriptorPool;
 class FVulkanDescriptorPoolsManager;
 class FVulkanBindlessDescriptorManager;
-class FVulkanCommandListContextImmediate;
 class FVulkanTransientHeapCache;
 class FVulkanDeviceExtension;
 class FVulkanRenderPassManager;
@@ -509,7 +508,7 @@ public:
 		return ShaderFactory;
 	}
 
-	FVulkanCommandListContextImmediate& GetImmediateContext()
+	FVulkanCommandListContext& GetImmediateContext()
 	{
 		return *ImmediateContext;
 	}
@@ -713,7 +712,7 @@ private:
 
 	TMap<uint32, FSamplerStateRHIRef> SamplerMap;
 
-	FVulkanCommandListContextImmediate* ImmediateContext;
+	FVulkanCommandListContext* ImmediateContext;
 
 	FVulkanDynamicRHI* RHI = nullptr;
 	
diff --git a/Engine/Source/Runtime/VulkanRHI/Private/VulkanDynamicRHI.h b/Engine/Source/Runtime/VulkanRHI/Private/VulkanDynamicRHI.h
--- a/Engine/Source/Runtime/VulkanRHI/Private/VulkanDynamicRHI.h
+++ b/Engine/Source/Runtime/VulkanRHI/Private/VulkanDynamicRHI.h
@@ -87,7 +87,7 @@ public:
 	virtual VkQueue RHIGetGraphicsVkQueue() const final override;
 	virtual uint32 RHIGetGraphicsQueueIndex() const final override;
 	virtual uint32 RHIGetGraphicsQueueFamilyIndex() const final override;
-	virtual VkCommandBuffer RHIGetActiveVkCommandBuffer() final override;
+	virtual VkCommandBuffer  RHIGetVkCommandBuffer(FRHICommandListBase& ExecutingCmdList) const final override;
 	virtual uint64 RHIGetGraphicsAdapterLUID(VkPhysicalDevice InPhysicalDevice) const final override;
 	virtual bool RHIDoesAdapterMatchDevice(const void* InAdapterId) const final override;
 	virtual void* RHIGetVkDeviceProcAddr(const char* InName) const final override;
@@ -109,11 +109,8 @@ public:
 	virtual FVulkanRHIAllocationInfo RHIGetAllocationInfo(FRHITexture* InTexture) const final override;
 	virtual FVulkanRHIImageViewInfo RHIGetImageViewInfo(FRHITexture* InTexture) const final override;
 	virtual FVulkanRHIAllocationInfo RHIGetAllocationInfo(FRHIBuffer* InBuffer) const final override;
-	virtual void RHISetImageLayout(VkImage Image, VkImageLayout OldLayout, VkImageLayout NewLayout, const VkImageSubresourceRange& SubresourceRange) final override;
-	virtual void RHISetUploadImageLayout(VkImage Image, VkImageLayout OldLayout, VkImageLayout NewLayout, const VkImageSubresourceRange& SubresourceRange) final override;
 	virtual void RHIFinishExternalComputeWork(VkCommandBuffer InCommandBuffer) final override;
 	virtual void RHIRegisterWork(uint32 NumPrimitives) final override;
-	virtual void RHISubmitUploadCommandBuffer() final override;
 	virtual void RHIVerifyResult(VkResult Result, const ANSICHAR* VkFuntion, const ANSICHAR* Filename, uint32 Line) final override;
 
 	// FDynamicRHI interface.
diff --git a/Engine/Source/Runtime/VulkanRHI/Private/VulkanRHI.cpp b/Engine/Source/Runtime/VulkanRHI/Private/VulkanRHI.cpp
--- a/Engine/Source/Runtime/VulkanRHI/Private/VulkanRHI.cpp
+++ b/Engine/Source/Runtime/VulkanRHI/Private/VulkanRHI.cpp
@@ -422,13 +422,6 @@ FDynamicRHI* FVulkanDynamicRHIModule::CreateRHI(ERHIFeatureLevel::Type InRequest
 
 IMPLEMENT_MODULE(FVulkanDynamicRHIModule, VulkanRHI);
 
-
-FVulkanCommandListContextImmediate::FVulkanCommandListContextImmediate(FVulkanDevice& InDevice)
-	: FVulkanCommandListContext(InDevice, ERHIPipeline::Graphics, nullptr)
-{
-}
-
-
 FVulkanDynamicRHI::FVulkanDynamicRHI()
 	: Instance(VK_NULL_HANDLE)
 	, Device(nullptr)
@@ -1424,9 +1417,8 @@ bool FVulkanDynamicRHI::RHIGetAvailableResolutions(FScreenResolutionArray& Resol
 
 void FVulkanDynamicRHI::RHIFlushResources()
 {
-	FVulkanCommandListContextImmediate& ImmediateContext = GetDevice()->GetImmediateContext();
 	bool bTrimMemory = true;
-	ImmediateContext.FreeUnusedCmdBuffers(bTrimMemory);
+	GetDevice()->GetImmediateContext().FreeUnusedCmdBuffers(bTrimMemory);
 }
 
 // IVulkanDynamicRHI interface
@@ -1480,11 +1472,9 @@ uint32 FVulkanDynamicRHI::RHIGetGraphicsQueueFamilyIndex() const
 	return GetDevice()->GetGraphicsQueue()->GetFamilyIndex();
 }
 
-VkCommandBuffer FVulkanDynamicRHI::RHIGetActiveVkCommandBuffer()
+VkCommandBuffer FVulkanDynamicRHI::RHIGetVkCommandBuffer(FRHICommandListBase& ExecutingCmdList) const
 {
-	FVulkanCommandListContextImmediate& ImmediateContext = GetDevice()->GetImmediateContext();
-	VkCommandBuffer VulkanCommandBuffer = ImmediateContext.GetActiveCmdBuffer()->GetHandle();
-	return VulkanCommandBuffer;
+	return FVulkanCommandListContext::Get(ExecutingCmdList).GetActiveCmdBuffer()->GetHandle();
 }
 
 uint64 FVulkanDynamicRHI::RHIGetGraphicsAdapterLUID(VkPhysicalDevice InPhysicalDevice) const
@@ -1648,23 +1638,9 @@ FVulkanRHIAllocationInfo FVulkanDynamicRHI::RHIGetAllocationInfo(FRHIBuffer* InB
 	return NewInfo;
 }
 
-void FVulkanDynamicRHI::RHISetImageLayout(VkImage Image, VkImageLayout OldLayout, VkImageLayout NewLayout, const VkImageSubresourceRange& SubresourceRange)
-{
-	FVulkanCommandListContext& ImmediateContext = GetDevice()->GetImmediateContext();
-	FVulkanCommandBuffer& CommandBuffer = ImmediateContext.GetCommandBuffer();
-	VulkanSetImageLayout(CommandBuffer.GetHandle(), Image, OldLayout, NewLayout, SubresourceRange);
-}
-
-void FVulkanDynamicRHI::RHISetUploadImageLayout(VkImage Image, VkImageLayout OldLayout, VkImageLayout NewLayout, const VkImageSubresourceRange& SubresourceRange)
-{
-	FVulkanCommandListContext& ImmediateContext = GetDevice()->GetImmediateContext();
-	FVulkanCommandBuffer& CommandBuffer = ImmediateContext.GetCommandBuffer();
-	VulkanSetImageLayout(CommandBuffer.GetHandle(), Image, OldLayout, NewLayout, SubresourceRange);
-}
-
 void FVulkanDynamicRHI::RHIFinishExternalComputeWork(VkCommandBuffer InCommandBuffer)
 {
-	FVulkanCommandListContextImmediate& ImmediateContext = GetDevice()->GetImmediateContext();
+	FVulkanCommandListContext& ImmediateContext = GetDevice()->GetImmediateContext();
 	check(InCommandBuffer == ImmediateContext.GetActiveCmdBuffer()->GetHandle());
 
 	ImmediateContext.GetPendingComputeState()->Reset();
@@ -1674,19 +1650,13 @@ void FVulkanDynamicRHI::RHIFinishExternalComputeWork(VkCommandBuffer InCommandBu
 void FVulkanDynamicRHI::RHIRegisterWork(uint32 NumPrimitives)
 {
 #if (RHI_NEW_GPU_PROFILER == 0)
-	FVulkanCommandListContextImmediate& ImmediateContext = GetDevice()->GetImmediateContext();
-	if (FVulkanPlatform::RegisterGPUWork() && ImmediateContext.IsImmediate())
+	if (FVulkanPlatform::RegisterGPUWork())
 	{
-		ImmediateContext.RegisterGPUWork(NumPrimitives);
+		GetDevice()->GetImmediateContext().RegisterGPUWork(NumPrimitives);
 	}
 #endif
 }
 
-void FVulkanDynamicRHI::RHISubmitUploadCommandBuffer()
-{
-
-}
-
 void FVulkanDynamicRHI::RHIVerifyResult(VkResult Result, const ANSICHAR* VkFuntion, const ANSICHAR* Filename, uint32 Line)
 {
 	VulkanRHI::VerifyVulkanResult(Result, VkFuntion, Filename, Line);
diff --git a/Engine/Source/Runtime/VulkanRHI/Private/VulkanSubmission.cpp b/Engine/Source/Runtime/VulkanRHI/Private/VulkanSubmission.cpp
--- a/Engine/Source/Runtime/VulkanRHI/Private/VulkanSubmission.cpp
+++ b/Engine/Source/Runtime/VulkanRHI/Private/VulkanSubmission.cpp
@@ -287,7 +287,7 @@ IRHIComputeContext* FVulkanDynamicRHI::RHIGetParallelCommandContext(FRHIParallel
 {
 	FVulkanParallelRenderPassInfo* VulkanParallelRenderPassInfo = (FVulkanParallelRenderPassInfo*)ParallelRenderPass.RHIPlatformData;
 	checkf(VulkanParallelRenderPassInfo, TEXT("Must begin parallel render pass on parent context before creating parallel contexts!"));
-	FVulkanCommandListContext* ParallelCommandContext = new FVulkanCommandListContext(*Device, &Device->GetImmediateContext(), VulkanParallelRenderPassInfo);
+	FVulkanCommandListContext* ParallelCommandContext = new FVulkanCommandListContext(*Device, VulkanParallelRenderPassInfo);
 	ParallelCommandContext->RenderPassInfo = ParallelRenderPass;
 
 	if (Device->GetOptionalExtensions().HasKHRDynamicRendering)
diff --git a/Engine/Source/Runtime/VulkanRHI/Private/VulkanUniformBuffer.cpp b/Engine/Source/Runtime/VulkanRHI/Private/VulkanUniformBuffer.cpp
--- a/Engine/Source/Runtime/VulkanRHI/Private/VulkanUniformBuffer.cpp
+++ b/Engine/Source/Runtime/VulkanRHI/Private/VulkanUniformBuffer.cpp
@@ -171,18 +171,16 @@ PRAGMA_ENABLE_DEPRECATION_WARNINGS
 				const bool bCanAllocOnThisThread = RHICmdList.Bypass() || (!IsRunningRHIInSeparateThread() && bInRenderingThread) || bInRHIThread;
 				if (bCanAllocOnThisThread)
 				{
-					FVulkanCommandListContextImmediate& Context = Device->GetImmediateContext();
-					UpdateUniformBufferHelper(Context, UniformBuffer, Contents);
+					UpdateUniformBufferHelper(Device->GetImmediateContext(), UniformBuffer, Contents);
 				}
 				else
 				{
 					void* CmdListConstantBufferData = RHICmdList.Alloc(DataSize, 16);
 					UpdateUniformBufferConstants(*Device, CmdListConstantBufferData, Contents, InLayout);
 
-					RHICmdList.EnqueueLambda([UniformBuffer, DataSize, CmdListConstantBufferData](FRHICommandList& CmdList)
+					RHICmdList.EnqueueLambda([UniformBuffer, DataSize, CmdListConstantBufferData](FRHICommandList& ExecutingCmdList)
 					{
-						FVulkanCommandListContext& Context = FVulkanCommandListContext::Get(CmdList);
-						UpdateUniformBufferHelper(Context, UniformBuffer, CmdListConstantBufferData, false);
+						UpdateUniformBufferHelper(FVulkanCommandListContext::Get(ExecutingCmdList), UniformBuffer, CmdListConstantBufferData, false);
 					});
 
 					RHICmdList.RHIThreadFence(true);
@@ -336,12 +334,9 @@ inline void FVulkanDynamicRHI::UpdateUniformBuffer(FRHICommandListBase& RHICmdLi
 			void* CmdListConstantBufferData = RHICmdList.Alloc(ConstantBufferSize, 16);
 			FMemory::Memcpy(CmdListConstantBufferData, Contents, ConstantBufferSize);
 
-			FVulkanCommandListContextImmediate* Context = &Device->GetImmediateContext();
-
-			RHICmdList.EnqueueLambda([UniformBuffer, CmdListResources, NumResources, ConstantBufferSize, CmdListConstantBufferData](FRHICommandListBase& CmdList)
+			RHICmdList.EnqueueLambda([UniformBuffer, CmdListResources, NumResources, ConstantBufferSize, CmdListConstantBufferData](FRHICommandListBase& ExecutingCmdList)
 			{
-				FVulkanCommandListContext& Context = (FVulkanCommandListContext&)CmdList.GetContext().GetLowestLevelContext();
-				UpdateUniformBufferHelper(Context, UniformBuffer, CmdListConstantBufferData);
+				UpdateUniformBufferHelper(FVulkanCommandListContext::Get(ExecutingCmdList), UniformBuffer, CmdListConstantBufferData);
 				UniformBuffer->UpdateResourceTable(CmdListResources, NumResources);
 			});
 		}
diff --git a/Engine/Source/Runtime/VulkanRHI/Private/VulkanViewport.cpp b/Engine/Source/Runtime/VulkanRHI/Private/VulkanViewport.cpp
--- a/Engine/Source/Runtime/VulkanRHI/Private/VulkanViewport.cpp
+++ b/Engine/Source/Runtime/VulkanRHI/Private/VulkanViewport.cpp
@@ -252,7 +252,6 @@ FTextureRHIRef FVulkanViewport::GetBackBuffer(FRHICommandListImmediate& RHICmdLi
 
 	if (SupportsStandardSwapchain() && GVulkanDelayAcquireImage != EDelayAcquireImageType::DelayAcquire)
 	{
-		check(RHICmdList.IsImmediate());
 		check(RHIBackBuffer);
 		
 		RHICmdList.EnqueueLambda([this](FRHICommandListImmediate& CmdList)
@@ -302,7 +301,7 @@ void FVulkanViewport::IssueFrameEvent()
 {
 	if (FVulkanPlatform::RequiresWaitingForFrameCompletionEvent())
 	{
-		FVulkanCommandListContextImmediate& ImmediateContext = Device.GetImmediateContext();
+		FVulkanCommandListContext& ImmediateContext = Device.GetImmediateContext();
 		LastFrameSyncPoint = ImmediateContext.GetContextSyncPoint();
 		ImmediateContext.FlushCommands();
 	}
diff --git a/Engine/Source/Runtime/VulkanRHI/Public/IVulkanDynamicRHI.h b/Engine/Source/Runtime/VulkanRHI/Public/IVulkanDynamicRHI.h
--- a/Engine/Source/Runtime/VulkanRHI/Public/IVulkanDynamicRHI.h
+++ b/Engine/Source/Runtime/VulkanRHI/Public/IVulkanDynamicRHI.h
@@ -51,7 +51,18 @@ struct IVulkanDynamicRHI : public FDynamicRHI
 	virtual uint32           RHIGetGraphicsQueueIndex() const = 0;
 	virtual uint32           RHIGetGraphicsQueueFamilyIndex() const = 0;
 
-	virtual VkCommandBuffer  RHIGetActiveVkCommandBuffer() = 0;
+	UE_DEPRECATED(5.8, "RHIGetActiveVkCommandBuffer is deprecated, and dangerous because it is incompatible with parallel translation. Use RHIGetVkCommandBuffer(FRHICommandListBase&).")
+	VkCommandBuffer RHIGetActiveVkCommandBuffer()
+	{
+		checkNoEntry();
+		return {};
+	}
+
+	//
+	// Returns the command buffer for the currently active pipeline on the given RHI command list.
+	// Can only be called on executing command lists (i.e. ExecutingCmdList.IsBottomOfPipe() must be true),
+	// 
+	virtual VkCommandBuffer  RHIGetVkCommandBuffer(FRHICommandListBase& ExecutingCmdList) const = 0;
 
 	virtual uint64           RHIGetGraphicsAdapterLUID(VkPhysicalDevice InPhysicalDevice) const = 0;
 	virtual bool             RHIDoesAdapterMatchDevice(const void* InAdapterId) const = 0;
@@ -82,16 +93,8 @@ struct IVulkanDynamicRHI : public FDynamicRHI
 	virtual FVulkanRHIImageViewInfo  RHIGetImageViewInfo(FRHITexture* InTexture) const = 0;
 	virtual FVulkanRHIAllocationInfo RHIGetAllocationInfo(FRHIBuffer* InBuffer) const = 0;
 
-	virtual void           RHISetImageLayout(VkImage Image, VkImageLayout OldLayout, VkImageLayout NewLayout, const VkImageSubresourceRange& SubresourceRange) = 0;
-
-	UE_DEPRECATED(5.5, "Upload command buffers are deprecated. Use RHISetImageLayout().")
-	virtual void           RHISetUploadImageLayout(VkImage Image, VkImageLayout OldLayout, VkImageLayout NewLayout, const VkImageSubresourceRange& SubresourceRange) = 0;
-
 	virtual void           RHIFinishExternalComputeWork(VkCommandBuffer InCommandBuffer) = 0;
 	virtual void           RHIRegisterWork(uint32 NumPrimitives) = 0;
-	
-	UE_DEPRECATED(5.5, "Upload command buffers are deprecated.")
-	virtual void           RHISubmitUploadCommandBuffer() = 0;
 
 	virtual void           RHIVerifyResult(VkResult Result, const ANSICHAR* VkFuntion, const ANSICHAR* Filename, uint32 Line) = 0;
 
