Remove legacy GPU profiler code paths from Vulkan RHI

diff --git a/Engine/Source/Runtime/VulkanRHI/Private/Android/VulkanAndroidPlatform.cpp b/Engine/Source/Runtime/VulkanRHI/Private/Android/VulkanAndroidPlatform.cpp
--- a/Engine/Source/Runtime/VulkanRHI/Private/Android/VulkanAndroidPlatform.cpp
+++ b/Engine/Source/Runtime/VulkanRHI/Private/Android/VulkanAndroidPlatform.cpp
@@ -388,14 +388,6 @@ void FVulkanAndroidPlatform::FreeVulkanLibrary()
 
 #undef CHECK_VK_ENTRYPOINTS
 
-bool FVulkanAndroidPlatform::HasCustomFrameTiming()
-{
-#if USE_ANDROID_VULKAN_SWAPPY
-	return FVulkanAndroidPlatform::bSwappyEnabledAtRHIInit;
-#endif
-	return false;
-}
-
 void FVulkanAndroidPlatform::InitDevice(FVulkanDevice* InDevice)
 {
 #if USE_ANDROID_VULKAN_SWAPPY
diff --git a/Engine/Source/Runtime/VulkanRHI/Private/Android/VulkanAndroidPlatform.h b/Engine/Source/Runtime/VulkanRHI/Private/Android/VulkanAndroidPlatform.h
--- a/Engine/Source/Runtime/VulkanRHI/Private/Android/VulkanAndroidPlatform.h
+++ b/Engine/Source/Runtime/VulkanRHI/Private/Android/VulkanAndroidPlatform.h
@@ -177,8 +177,6 @@ public:
 		#endif
 	}
 
-	static bool HasCustomFrameTiming();
-
 	static bool SupportsVolumeTextureRendering() { return false; }
 
 	static void OverridePlatformHandlers(bool bInit);
@@ -192,8 +190,6 @@ public:
 	static bool HasUnifiedMemory() { return true; }
 #endif
 
-	static bool RegisterGPUWork() { return false; }
-
 	// Assume most devices can't use the extra cores for running parallel tasks
 	static bool SupportParallelRenderingTasks() { return false; }
 
diff --git a/Engine/Source/Runtime/VulkanRHI/Private/VulkanCommandBuffer.cpp b/Engine/Source/Runtime/VulkanRHI/Private/VulkanCommandBuffer.cpp
--- a/Engine/Source/Runtime/VulkanRHI/Private/VulkanCommandBuffer.cpp
+++ b/Engine/Source/Runtime/VulkanRHI/Private/VulkanCommandBuffer.cpp
@@ -21,9 +21,7 @@
 FVulkanCommandBuffer::FVulkanCommandBuffer(FVulkanDevice& InDevice, FVulkanCommandBufferPool& InCommandBufferPool)
 	: Device(InDevice)
 	, CommandBufferPool(InCommandBufferPool)
-#if RHI_NEW_GPU_PROFILER
 	, EventStream(InCommandBufferPool.GetQueue().GetProfilerQueue())
-#endif
 {
 	FScopeLock ScopeLock(CommandBufferPool.GetCS());
 	AllocMemory();
@@ -155,20 +153,12 @@ void FVulkanCommandBuffer::End(FVulkanQueryPool* OptionalQueryPool)
 		VulkanRHI::vkCmdResetEvent(GetHandle(), BarrierEvent, VK_PIPELINE_STAGE_ALL_COMMANDS_BIT);
 	}
 
-#if RHI_NEW_GPU_PROFILER
 	if (OptionalQueryPool)
 	{
 		auto& Event = EmplaceProfilerEvent<UE::RHI::GPUProfiler::FEvent::FEndWork>(0);
 		const uint32 IndexInPool = OptionalQueryPool->ReserveQuery(&Event.GPUTimestampBOP);
 		VulkanRHI::vkCmdWriteTimestamp(GetHandle(), VK_PIPELINE_STAGE_BOTTOM_OF_PIPE_BIT, OptionalQueryPool->GetHandle(), IndexInPool);
 	}
-#else
-	if (OptionalQueryPool)
-	{
-		const uint32 IndexInPool = OptionalQueryPool->ReserveQuery(&EndTimestamp);
-		VulkanRHI::vkCmdWriteTimestamp(GetHandle(), VK_PIPELINE_STAGE_BOTTOM_OF_PIPE_BIT, OptionalQueryPool->GetHandle(), IndexInPool);
-	}
-#endif
 
 	VERIFYVULKANRESULT(VulkanRHI::vkEndCommandBuffer(GetHandle()));
 	State = EState::HasEnded;
@@ -215,22 +205,12 @@ void FVulkanCommandBuffer::Begin(FVulkanQueryPool* OptionalQueryPool, VkRenderPa
 		BindlessDescriptorManager->BindDescriptorBuffers(CommandBufferHandle, SupportedStages);
 	}
 
-#if RHI_NEW_GPU_PROFILER
 	if (OptionalQueryPool)
 	{
 		auto& Event = EmplaceProfilerEvent<UE::RHI::GPUProfiler::FEvent::FBeginWork>(0, UINT64_MAX);
 		const uint32 IndexInPool = OptionalQueryPool->ReserveQuery(&Event.GPUTimestampTOP);
 		VulkanRHI::vkCmdWriteTimestamp(GetHandle(), VK_PIPELINE_STAGE_TOP_OF_PIPE_BIT, OptionalQueryPool->GetHandle(), IndexInPool);
 	}
-#else
-	if (OptionalQueryPool)
-	{
-		StartTimestamp = 0;
-		EndTimestamp = 0;
-		const uint32 IndexInPool = OptionalQueryPool->ReserveQuery(&StartTimestamp);
-		VulkanRHI::vkCmdWriteTimestamp(GetHandle(), VK_PIPELINE_STAGE_TOP_OF_PIPE_BIT, OptionalQueryPool->GetHandle(), IndexInPool);
-	}
-#endif
 
 	bNeedsFullDynamicStateUpdate = true;
 }
diff --git a/Engine/Source/Runtime/VulkanRHI/Private/VulkanCommandBuffer.h b/Engine/Source/Runtime/VulkanRHI/Private/VulkanCommandBuffer.h
--- a/Engine/Source/Runtime/VulkanRHI/Private/VulkanCommandBuffer.h
+++ b/Engine/Source/Runtime/VulkanRHI/Private/VulkanCommandBuffer.h
@@ -135,7 +135,6 @@ public:
 	void BeginSplitBarrier(VkEvent BarrierEvent, const VkDependencyInfo& DependencyInfo);
 	void EndSplitBarrier(VkEvent BarrierEvent, const VkDependencyInfo& DependencyInfo);
 
-#if RHI_NEW_GPU_PROFILER
 	template <typename TEventType, typename... TArgs>
 	TEventType& EmplaceProfilerEvent(TArgs&&... Args)
 	{
@@ -161,16 +160,6 @@ public:
 
 		Destination.Append(MoveTemp(EventStream));
 	}
-#else
-	uint64 GetBusyCycles() const
-	{
-		if (EndTimestamp > StartTimestamp)
-		{
-			return (EndTimestamp - StartTimestamp);
-		}
-		return 0;
-	}
-#endif
 
 	// Public immutable reference to be accessed directly
 	FVulkanDevice& Device;
@@ -190,13 +179,8 @@ private:
 	void AllocMemory();
 	void FreeMemory();
 
-#if RHI_NEW_GPU_PROFILER
 	UE::RHI::GPUProfiler::FEventStream EventStream;
 	TArray<UE::RHI::GPUProfiler::FEvent::FBeginWork*, TInlineAllocator<8>> BeginEvents;
-#else
-	uint64 StartTimestamp = 0;
-	uint64 EndTimestamp = 0;
-#endif
 
 	friend class FVulkanDynamicRHI;
 	friend FVulkanContextCommon;
diff --git a/Engine/Source/Runtime/VulkanRHI/Private/VulkanCommands.cpp b/Engine/Source/Runtime/VulkanRHI/Private/VulkanCommands.cpp
--- a/Engine/Source/Runtime/VulkanRHI/Private/VulkanCommands.cpp
+++ b/Engine/Source/Runtime/VulkanRHI/Private/VulkanCommands.cpp
@@ -293,13 +293,6 @@ void FVulkanCommandListContext::RHIDispatchComputeShader(uint32 ThreadGroupCount
 
 	VulkanRHI::vkCmdDispatch(CommandBufferHandle, ThreadGroupCountX, ThreadGroupCountY, ThreadGroupCountZ);
 
-#if (RHI_NEW_GPU_PROFILER == 0)
-	if (FVulkanPlatform::RegisterGPUWork() && IsImmediate())
-	{
-		RegisterGPUDispatch(FIntVector(ThreadGroupCountX, ThreadGroupCountY, ThreadGroupCountZ));	
-	}
-#endif
-
 	VulkanRHI::DebugHeavyWeightBarrier(CommandBufferHandle, 2);
 }
 
@@ -319,13 +312,6 @@ void FVulkanCommandListContext::RHIDispatchIndirectComputeShader(FRHIBuffer* Arg
 
 	VulkanRHI::vkCmdDispatchIndirect(CommandBufferHandle, ArgumentBuffer->GetHandle(), ArgumentBuffer->GetOffset() + ArgumentOffset);
 
-#if (RHI_NEW_GPU_PROFILER == 0)
-	if (FVulkanPlatform::RegisterGPUWork()/* && IsImmediate()*/)
-	{
-		RegisterGPUDispatch(FIntVector(1, 1, 1));	
-	}
-#endif
-
 	VulkanRHI::DebugHeavyWeightBarrier(CommandBufferHandle, 2);
 }
 
@@ -650,13 +636,6 @@ void FVulkanCommandListContext::RHIDrawPrimitive(uint32 BaseVertexIndex, uint32
 	RHI_DRAW_CALL_STATS(PendingGfxState->PrimitiveType, NumVertices, NumPrimitives, NumInstances);
 
 	VulkanRHI::vkCmdDraw(GetCommandBuffer().GetHandle(), NumVertices, NumInstances, BaseVertexIndex, 0);
-
-#if (RHI_NEW_GPU_PROFILER == 0)
-	if (FVulkanPlatform::RegisterGPUWork() && IsImmediate())
-	{
-		RegisterGPUWork(NumPrimitives * NumInstances, NumVertices * NumInstances);
-	}
-#endif
 }
 
 void FVulkanCommandListContext::RHIDrawPrimitiveIndirect(FRHIBuffer* ArgumentBufferRHI, uint32 ArgumentOffset)
@@ -677,13 +656,6 @@ void FVulkanCommandListContext::RHIDrawPrimitiveIndirect(FRHIBuffer* ArgumentBuf
 	FVulkanBuffer* ArgumentBuffer = ResourceCast(ArgumentBufferRHI);
 
 	VulkanRHI::vkCmdDrawIndirect(CommandBufferHandle, ArgumentBuffer->GetHandle(), ArgumentBuffer->GetOffset() + ArgumentOffset, 1, sizeof(VkDrawIndirectCommand));
-
-#if (RHI_NEW_GPU_PROFILER == 0)
-	if (FVulkanPlatform::RegisterGPUWork() && IsImmediate())
-	{
-		RegisterGPUWork(1);
-	}
-#endif
 }
 
 void FVulkanCommandListContext::RHIDrawIndexedPrimitive(FRHIBuffer* IndexBufferRHI, int32 BaseVertexIndex, uint32 FirstInstance,
@@ -706,13 +678,6 @@ void FVulkanCommandListContext::RHIDrawIndexedPrimitive(FRHIBuffer* IndexBufferR
 
 	uint32 NumIndices = GetVertexCountForPrimitiveCount(NumPrimitives, PendingGfxState->PrimitiveType);
 	VulkanRHI::vkCmdDrawIndexed(CommandBufferHandle, NumIndices, NumInstances, StartIndex, BaseVertexIndex, FirstInstance);
-
-#if (RHI_NEW_GPU_PROFILER == 0)
-	if (FVulkanPlatform::RegisterGPUWork() && IsImmediate())
-	{
-		RegisterGPUWork(NumPrimitives * NumInstances, NumVertices * NumInstances);
-	}
-#endif
 }
 
 void FVulkanCommandListContext::RHIDrawIndexedPrimitiveIndirect(FRHIBuffer* IndexBufferRHI, FRHIBuffer* ArgumentBufferRHI, uint32 ArgumentOffset)
@@ -732,15 +697,7 @@ void FVulkanCommandListContext::RHIDrawIndexedPrimitiveIndirect(FRHIBuffer* Inde
 
 	FVulkanBuffer* ArgumentBuffer = ResourceCast(ArgumentBufferRHI);
 
-
 	VulkanRHI::vkCmdDrawIndexedIndirect(CommandBufferHandle, ArgumentBuffer->GetHandle(), ArgumentBuffer->GetOffset() + ArgumentOffset, 1, sizeof(VkDrawIndexedIndirectCommand));
-
-#if (RHI_NEW_GPU_PROFILER == 0)
-	if (FVulkanPlatform::RegisterGPUWork() && IsImmediate())
-	{
-		RegisterGPUWork(1); 
-	}
-#endif
 }
 
 #if PLATFORM_SUPPORTS_MESH_SHADERS
@@ -753,13 +710,6 @@ void FVulkanCommandListContext::RHIDispatchMeshShader(uint32 ThreadGroupCountX,
 	PendingGfxState->PrepareForDraw(*this);
 
 	VulkanRHI::vkCmdDrawMeshTasksEXT(GetCommandBuffer().GetHandle(), ThreadGroupCountX, ThreadGroupCountY, ThreadGroupCountZ);
-
-#if (RHI_NEW_GPU_PROFILER == 0)
-	if (FVulkanPlatform::RegisterGPUWork() && IsImmediate())
-	{
-		RegisterGPUWork(1);
-	}
-#endif
 }
 
 void FVulkanCommandListContext::RHIDispatchIndirectMeshShader(FRHIBuffer* ArgumentBufferRHI, uint32 ArgumentOffset)
@@ -773,13 +723,6 @@ void FVulkanCommandListContext::RHIDispatchIndirectMeshShader(FRHIBuffer* Argume
 	FVulkanBuffer* ArgumentBuffer = ResourceCast(ArgumentBufferRHI);
 
 	VulkanRHI::vkCmdDrawMeshTasksIndirectEXT(GetCommandBuffer().GetHandle(), ArgumentBuffer->GetHandle(), ArgumentBuffer->GetOffset() + ArgumentOffset, 1, sizeof(VkDrawMeshTasksIndirectCommandEXT));
-
-#if (RHI_NEW_GPU_PROFILER == 0)
-	if (FVulkanPlatform::RegisterGPUWork() && IsImmediate())
-	{
-		RegisterGPUWork(1);
-	}
-#endif
 }
 #endif
 
diff --git a/Engine/Source/Runtime/VulkanRHI/Private/VulkanContext.cpp b/Engine/Source/Runtime/VulkanRHI/Private/VulkanContext.cpp
--- a/Engine/Source/Runtime/VulkanRHI/Private/VulkanContext.cpp
+++ b/Engine/Source/Runtime/VulkanRHI/Private/VulkanContext.cpp
@@ -64,9 +64,10 @@ void FVulkanContextCommon::EndPayload()
 				}
 
 				// Only record begin/end timestamps on primary command buffers
-				FVulkanQueryPool* OptionalTimestampQueryPool =
-					GRHIGlobals.SupportsTimestampRenderQueries && bIsPrimaryCommandBuffer ?
-					GetCurrentTimestampQueryPool(*Payload) : nullptr;
+				FVulkanQueryPool* OptionalTimestampQueryPool = GRHIGlobals.SupportsTimestampRenderQueries && bIsPrimaryCommandBuffer
+					? GetCurrentTimestampQueryPool(*Payload)
+					: nullptr;
+
 				CommandBuffer->End(OptionalTimestampQueryPool);
 			}
 		}
@@ -208,9 +209,15 @@ void FVulkanContextCommon::PrepareNewCommandBuffer(FVulkanPayload& Payload)
 
 	// Only record begin/end timestamps on primary command buffers
 	const bool bIsPrimaryCommandBuffer = (NewCommandBuffer->GetCommandBufferType() != EVulkanCommandBufferType::Secondary);
-	FVulkanQueryPool* OptionalTimestampQueryPool = GRHIGlobals.SupportsTimestampRenderQueries && bIsPrimaryCommandBuffer ?
-		GetCurrentTimestampQueryPool(Payload) : nullptr;
-	VkRenderPass RenderPassHandle = GetParallelRenderPassInfo() ? GetParallelRenderPassInfo()->RenderPassHandle : VK_NULL_HANDLE;
+	
+	FVulkanQueryPool* OptionalTimestampQueryPool = GRHIGlobals.SupportsTimestampRenderQueries && bIsPrimaryCommandBuffer
+		? GetCurrentTimestampQueryPool(Payload)
+		: nullptr;
+
+	VkRenderPass RenderPassHandle = GetParallelRenderPassInfo()
+		? GetParallelRenderPassInfo()->RenderPassHandle
+		: VK_NULL_HANDLE;
+
 	NewCommandBuffer->Begin(OptionalTimestampQueryPool, RenderPassHandle);
 }
 
@@ -258,15 +265,7 @@ FVulkanCommandListContext::FVulkanCommandListContext(FVulkanDevice& InDevice, ER
 	, RHIPipeline(InPipeline)
 	, bSupportsBreadcrumbs(GRHIGlobals.SupportsTimestampRenderQueries)
 	, bImmediate(bInImmediate)
-#if (RHI_NEW_GPU_PROFILER == 0)
-	, GpuProfiler(this, &InDevice)
-#endif
 {
-#if (RHI_NEW_GPU_PROFILER == 0)
-	FrameTiming = new FVulkanGPUTiming(this, &InDevice);
-	FrameTiming->Initialize();
-#endif
-
 	// Create Pending state, contains pipeline states such as current shader and etc..
 	PendingGfxState = new FVulkanPendingGfxState(Device);
 	PendingComputeState = new FVulkanPendingComputeState(Device);
@@ -280,9 +279,6 @@ FVulkanCommandListContext::FVulkanCommandListContext(FVulkanDevice& InDevice, FV
 	, bSupportsBreadcrumbs(false)
 	, bImmediate(false)
 	, CurrentParallelRenderPassInfo(InParallelRenderPassInfo)
-#if (RHI_NEW_GPU_PROFILER == 0)
-	, GpuProfiler(this, &InDevice)
-#endif
 {
 	checkf(CurrentParallelRenderPassInfo, TEXT("Secondary command buffers should be created with a FVulkanParallelRenderPassInfo."));
 
@@ -312,18 +308,6 @@ void FVulkanCommandListContext::ReleasePendingState()
 
 FVulkanCommandListContext::~FVulkanCommandListContext()
 {
-	if (GSupportsTimestampRenderQueries)
-	{
-#if (RHI_NEW_GPU_PROFILER == 0)
-		if (FrameTiming)
-		{
-			FrameTiming->Release();
-			delete FrameTiming;
-			FrameTiming = nullptr;
-		}
-#endif
-	}
-
 	ReleasePendingState();
 }
 
diff --git a/Engine/Source/Runtime/VulkanRHI/Private/VulkanContext.h b/Engine/Source/Runtime/VulkanRHI/Private/VulkanContext.h
--- a/Engine/Source/Runtime/VulkanRHI/Private/VulkanContext.h
+++ b/Engine/Source/Runtime/VulkanRHI/Private/VulkanContext.h
@@ -8,7 +8,6 @@
 
 #include "VulkanResources.h"
 #include "VulkanRHIPrivate.h"
-#include "VulkanGPUProfiler.h"
 #include "VulkanSubmission.h"
 #include "VulkanSynchronization.h"
 
@@ -321,7 +320,6 @@ public:
 	virtual void RHIBeginRenderQuery(FRHIRenderQuery* RenderQuery) final override;
 	virtual void RHIEndRenderQuery(FRHIRenderQuery* RenderQuery) final override;
 
-#if RHI_NEW_GPU_PROFILER
 	void FlushProfilerStats()
 	{
 		// Flush accumulated draw stats (if breadcrumbs are available to attach them to)
@@ -331,9 +329,6 @@ public:
 			StatEvent = {};
 		}
 	}
-#else
-	virtual void RHICalibrateTimers(FRHITimestampCalibrationQuery* CalibrationQuery) final override;
-#endif
 
 	virtual void Finalize(TArray<FVulkanPayload*>& OutPayloads) override;
 
@@ -386,13 +381,6 @@ public:
 		return CurrentFramebuffer;
 	}
 
-#if (RHI_NEW_GPU_PROFILER == 0)
-	inline FVulkanGPUProfiler& GetGPUProfiler()
-	{
-		return GpuProfiler;
-	}
-#endif
-
 	void BeginRecursiveCommand()
 	{
 		// Nothing to do
@@ -463,14 +451,6 @@ private:
 	friend class FVulkanDevice;
 	friend class FVulkanDynamicRHI;
 
-#if (RHI_NEW_GPU_PROFILER == 0)
-	void RegisterGPUWork(uint32 NumPrimitives = 0, uint32 NumVertices = 0)	{ GpuProfiler.RegisterGPUWork(NumPrimitives, NumVertices); }
-	void RegisterGPUDispatch(FIntVector GroupCount)	                        { GpuProfiler.RegisterGPUDispatch(GroupCount); }
-
-	FVulkanGPUProfiler GpuProfiler;
-	FVulkanGPUTiming* FrameTiming = nullptr;
-#endif
-
 	template <typename TRHIShader>
 	void ApplyStaticUniformBuffers(TRHIShader* Shader);
 
diff --git a/Engine/Source/Runtime/VulkanRHI/Private/VulkanDevice.cpp b/Engine/Source/Runtime/VulkanRHI/Private/VulkanDevice.cpp
--- a/Engine/Source/Runtime/VulkanRHI/Private/VulkanDevice.cpp
+++ b/Engine/Source/Runtime/VulkanRHI/Private/VulkanDevice.cpp
@@ -1610,7 +1610,6 @@ FVulkanTransientHeapCache& FVulkanDevice::GetOrCreateTransientHeapCache()
 	return *TransientHeapCache;
 }
 
-#if RHI_NEW_GPU_PROFILER
 void FVulkanDevice::GetCalibrationTimestamp(FVulkanTiming& InOutTiming)
 {
 	// TimestampPeriod is the number of nanoseconds required for a timestamp query to be incremented by 1
@@ -1662,34 +1661,6 @@ void FVulkanDevice::GetCalibrationTimestamp(FVulkanTiming& InOutTiming)
 		InOutTiming.CPUTimestamp = FPlatformTime::Cycles64();
 	}
 }
-#else
-FGPUTimingCalibrationTimestamp FVulkanDevice::GetCalibrationTimestamp()
-{
-	auto ToMicroseconds = [](uint64_t Timestamp)
-	{
-		const double Frequency = double(FVulkanGPUTiming::GetTimingFrequency());
-		uint64 Microseconds = (uint64)((double(Timestamp) / Frequency) * 1000.0 * 1000.0);
-		return Microseconds;
-	};
-
-	FGPUTimingCalibrationTimestamp CalibrationTimestamp;
-	if (OptionalDeviceExtensions.HasEXTCalibratedTimestamps)
-	{
-		VkCalibratedTimestampInfoEXT TimestampInfo;
-		ZeroVulkanStruct(TimestampInfo, VK_STRUCTURE_TYPE_CALIBRATED_TIMESTAMP_INFO_EXT);
-		TimestampInfo.timeDomain = VK_TIME_DOMAIN_DEVICE_EXT;
-
-		uint64_t GPUTimestamp = 0;
-		uint64_t MaxDeviation = 0;
-		VERIFYVULKANRESULT(VulkanRHI::vkGetCalibratedTimestampsEXT(Device, 1, &TimestampInfo, &GPUTimestamp, &MaxDeviation));
-		CalibrationTimestamp.GPUMicroseconds = ToMicroseconds(GPUTimestamp);
-
-		const uint64 CPUTimestamp = FPlatformTime::Cycles64();
-		CalibrationTimestamp.CPUMicroseconds = uint64(FPlatformTime::ToSeconds64(CPUTimestamp) * 1e6);
-	}
-	return CalibrationTimestamp;
-}
-#endif // (RHI_NEW_GPU_PROFILER == 0)
 
 void FVulkanDevice::ForEachQueue(TFunctionRef<void(FVulkanQueue&)> Callback)
 {
diff --git a/Engine/Source/Runtime/VulkanRHI/Private/VulkanDevice.h b/Engine/Source/Runtime/VulkanRHI/Private/VulkanDevice.h
--- a/Engine/Source/Runtime/VulkanRHI/Private/VulkanDevice.h
+++ b/Engine/Source/Runtime/VulkanRHI/Private/VulkanDevice.h
@@ -608,13 +608,7 @@ public:
 
 	const TArray<const ANSICHAR*>& GetDeviceExtensions() { return DeviceExtensions; }
 
-#if RHI_NEW_GPU_PROFILER
 	void GetCalibrationTimestamp(FVulkanTiming& InOutTiming);
-#else
-	// Performs a GPU and CPU timestamp at nearly the same time.
-	// This allows aligning GPU and CPU events on the same timeline in profile visualization.
-	FGPUTimingCalibrationTimestamp GetCalibrationTimestamp();
-#endif
 
 	const FVulkanSamplerState& GetGlobalSamplers(FVulkanShaderHeader::EGlobalSamplerType Type) const
 	{
diff --git a/Engine/Source/Runtime/VulkanRHI/Private/VulkanDynamicRHI.h b/Engine/Source/Runtime/VulkanRHI/Private/VulkanDynamicRHI.h
--- a/Engine/Source/Runtime/VulkanRHI/Private/VulkanDynamicRHI.h
+++ b/Engine/Source/Runtime/VulkanRHI/Private/VulkanDynamicRHI.h
@@ -45,7 +45,6 @@ struct FOptionalVulkanInstanceExtensions
 	}
 };
 
-#if RHI_NEW_GPU_PROFILER
 // Encapsulates the state required for tracking GPU queue performance across a frame.
 class FVulkanTiming
 {
@@ -60,7 +59,6 @@ public:
 
 	UE::RHI::GPUProfiler::FEventStream EventStream;
 };
-#endif
 
 PRAGMA_DISABLE_DEPRECATION_WARNINGS
 
@@ -110,7 +108,6 @@ public:
 	virtual FVulkanRHIImageViewInfo RHIGetImageViewInfo(FRHITexture* InTexture) const final override;
 	virtual FVulkanRHIAllocationInfo RHIGetAllocationInfo(FRHIBuffer* InBuffer) const final override;
 	virtual void RHIFinishExternalComputeWork(VkCommandBuffer InCommandBuffer) final override;
-	virtual void RHIRegisterWork(uint32 NumPrimitives) final override;
 	virtual void RHIVerifyResult(VkResult Result, const ANSICHAR* VkFuntion, const ANSICHAR* Filename, uint32 Line) final override;
 
 	// FDynamicRHI interface.
@@ -361,7 +358,6 @@ protected:
 	bool WaitAndProcessInterruptQueue();
 	void KickInterruptThread();
 
-#if RHI_NEW_GPU_PROFILER
 	struct FVulkanTimingArray : public TArray<TUniquePtr<FVulkanTiming>, TInlineAllocator<(int32)EVulkanQueueType::Count>>
 	{
 		FVulkanTiming* CreateNew(FVulkanQueue& Queue)
@@ -370,7 +366,6 @@ protected:
 		}
 	};
 	FVulkanTimingArray CurrentTimingPerQueue;
-#endif
 
 	FCriticalSection SubmissionCS;
 	FVulkanThread* SubmissionThread = nullptr;
diff --git a/Engine/Source/Runtime/VulkanRHI/Private/VulkanGPUProfiler.h b/Engine/Source/Runtime/VulkanRHI/Private/VulkanGPUProfiler.h
deleted file mode 100644
diff --git a/Engine/Source/Runtime/VulkanRHI/Private/VulkanGenericPlatform.h b/Engine/Source/Runtime/VulkanRHI/Private/VulkanGenericPlatform.h
--- a/Engine/Source/Runtime/VulkanRHI/Private/VulkanGenericPlatform.h
+++ b/Engine/Source/Runtime/VulkanRHI/Private/VulkanGenericPlatform.h
@@ -78,8 +78,6 @@ public:
 
 	static bool SupportsTimestampRenderQueries() { return true; }
 
-	static bool HasCustomFrameTiming() { return false; }
-
 	static bool RequiresMobileRenderer() { return false; }
 
 	static ERHIFeatureLevel::Type GetFeatureLevel(ERHIFeatureLevel::Type RequestedFeatureLevel);
@@ -94,8 +92,6 @@ public:
 
 	static bool HasUnifiedMemory() { return false; }
 
-	static bool RegisterGPUWork() { return true; }
-
 	// Allow the platform code to restrict the device features
 	static void RestrictEnabledPhysicalDeviceFeatures(FVulkanPhysicalDeviceFeatures* InOutFeaturesToEnable);
 
diff --git a/Engine/Source/Runtime/VulkanRHI/Private/VulkanQuery.cpp b/Engine/Source/Runtime/VulkanRHI/Private/VulkanQuery.cpp
--- a/Engine/Source/Runtime/VulkanRHI/Private/VulkanQuery.cpp
+++ b/Engine/Source/Runtime/VulkanRHI/Private/VulkanQuery.cpp
@@ -365,18 +365,6 @@ void FVulkanCommandListContext::RHIEndRenderQuery(FRHIRenderQuery* QueryRHI)
 	}
 }
 
-#if (RHI_NEW_GPU_PROFILER == 0)
-void FVulkanCommandListContext::RHICalibrateTimers(FRHITimestampCalibrationQuery* CalibrationQuery)
-{
-	if (Device.GetOptionalExtensions().HasEXTCalibratedTimestamps)
-	{
-		FGPUTimingCalibrationTimestamp CalibrationTimestamp = Device.GetCalibrationTimestamp();
-		CalibrationQuery->CPUMicroseconds[0] = CalibrationTimestamp.CPUMicroseconds;
-		CalibrationQuery->GPUMicroseconds[0] = CalibrationTimestamp.GPUMicroseconds;
-	}
-}
-#endif
-
 FVulkanQueryPool* FVulkanCommandListContext::GetCurrentOcclusionQueryPool()
 {
 	TArray<FVulkanQueryPool*>& OcclusionPoolArray = GetQueryPoolArray(EVulkanQueryPoolType::Occlusion);
diff --git a/Engine/Source/Runtime/VulkanRHI/Private/VulkanQueue.cpp b/Engine/Source/Runtime/VulkanRHI/Private/VulkanQueue.cpp
--- a/Engine/Source/Runtime/VulkanRHI/Private/VulkanQueue.cpp
+++ b/Engine/Source/Runtime/VulkanRHI/Private/VulkanQueue.cpp
@@ -182,7 +182,6 @@ int32 FVulkanQueue::SubmitQueuedPayloads(TMap<VkSemaphore, FBinarySemaphoreSigna
 			}
 		}
 
-#if RHI_NEW_GPU_PROFILER
 		// Find the maximum timeline value to wait on for each queue
 		uint64 MaxTimelineWaitValue[(int32)EVulkanQueueType::Count] = {};
 		for (FVulkanSemaphore* WaitSema : WaitSemas)
@@ -209,7 +208,6 @@ int32 FVulkanQueue::SubmitQueuedPayloads(TMap<VkSemaphore, FBinarySemaphoreSigna
 					SubmitTime, MaxTimelineWaitValue[QueueTypeIndex], SignalQueue->GetProfilerQueue());
 			}
 		}
-#endif
 
 		// We can only remove them from the list if they are all present
 		for (FVulkanSemaphore* WaitSema : WaitSemas)
@@ -322,7 +320,6 @@ void FVulkanQueue::SubmitPayloads(TArrayView<FVulkanPayload*> Payloads, TMap<VkS
 	{
 		Payload->PreExecute();
 
-#if RHI_NEW_GPU_PROFILER
 		const bool bHasExternalGPUTime = Payload->ExternalGPUTime.IsSet();
 		if (bHasExternalGPUTime)
 		{
@@ -343,7 +340,6 @@ void FVulkanQueue::SubmitPayloads(TArrayView<FVulkanPayload*> Payloads, TMap<VkS
 				Device.GetCalibrationTimestamp(*LocalTiming);
 			}
 		}
-#endif // RHI_NEW_GPU_PROFILER
 
 		// Some payloads have nothing to submit because they are only used to trigger CPU events
 		if (!Payload->WaitSemaphores.Num() && !Payload->CommandBuffers.Num() && !Payload->SignalSemaphores.Num() && !Payload->ReservedResourcesToCommit.Num())
@@ -396,13 +392,13 @@ void FVulkanQueue::SubmitPayloads(TArrayView<FVulkanPayload*> Payloads, TMap<VkS
 			for (FVulkanCommandBuffer* CommandBuffer : Payload->CommandBuffers)
 			{
 				CommandBufferStorage.Add(CommandBuffer->GetHandle());
-#if RHI_NEW_GPU_PROFILER
+
 				CommandBuffer->FlushProfilerEvents(Payload->EventStream, SubmitTime);
 				for (FVulkanCommandBuffer* SecondaryCommandBuffer : CommandBuffer->ExecutedSecondaryCommandBuffers)
 				{
 					SecondaryCommandBuffer->FlushProfilerEvents(Payload->EventStream, SubmitTime);
 				}
-#endif // RHI_NEW_GPU_PROFILER
+
 				CommandBuffer->SetSubmitted();
 			}
 			SubmitInfo.commandBufferCount = Payload->CommandBuffers.Num();
@@ -422,10 +418,9 @@ void FVulkanQueue::SubmitPayloads(TArrayView<FVulkanPayload*> Payloads, TMap<VkS
 			SubmitInfo.pSignalSemaphores = &SemaphoreStorage[FirstSempahoreIndex];
 		}
 
-#if RHI_NEW_GPU_PROFILER
 		Payload->EventStream.Emplace<UE::RHI::GPUProfiler::FEvent::FSignalFence>(
-			SubmitTime, NextTimelineSemaphoreValue);
-#endif
+			SubmitTime, NextTimelineSemaphoreValue
+		);
 
 		if (bUseTimelineSemaphores)
 		{
@@ -666,7 +661,6 @@ void FVulkanQueue::FillSupportedStageBits()
 	}
 }
 
-#if RHI_NEW_GPU_PROFILER
 UE::RHI::GPUProfiler::FQueue FVulkanQueue::GetProfilerQueue() const
 {
 	UE::RHI::GPUProfiler::FQueue ProfilerQueue;
@@ -683,7 +677,6 @@ UE::RHI::GPUProfiler::FQueue FVulkanQueue::GetProfilerQueue() const
 
 	return ProfilerQueue;
 }
-#endif // RHI_NEW_GPU_PROFILER
 
 FVulkanCommandBufferPool* FVulkanQueue::AcquireCommandBufferPool(EVulkanCommandBufferType CommandBufferType)
 {
diff --git a/Engine/Source/Runtime/VulkanRHI/Private/VulkanQueue.h b/Engine/Source/Runtime/VulkanRHI/Private/VulkanQueue.h
--- a/Engine/Source/Runtime/VulkanRHI/Private/VulkanQueue.h
+++ b/Engine/Source/Runtime/VulkanRHI/Private/VulkanQueue.h
@@ -111,9 +111,7 @@ public:
 	int32 SubmitQueuedPayloads(TMap<VkSemaphore, FBinarySemaphoreSignalInfo>& SignaledSemas);
 	int32 ProcessInterruptQueue(uint64 Timeout);
 
-#if RHI_NEW_GPU_PROFILER
 	UE::RHI::GPUProfiler::FQueue GetProfilerQueue() const;
-#endif
 
 	void InitDiagnosticBuffer();
 	class FVulkanDiagnosticBuffer* GetDiagnosticBuffer()
@@ -157,10 +155,8 @@ private:
 	TQueue<FVulkanPayload*, EQueueMode::Mpsc> PendingSubmission;
 	TQueue<FVulkanPayload*, EQueueMode::Spsc> PendingInterrupt;
 
-#if RHI_NEW_GPU_PROFILER
 	// The active timing struct on this queue. Updated / accessed by the interrupt thread.
 	FVulkanTiming* Timing = nullptr;
-#endif
 
 	TUniquePtr<class FVulkanDiagnosticBuffer> DiagnosticBuffer;
 
diff --git a/Engine/Source/Runtime/VulkanRHI/Private/VulkanRHI.cpp b/Engine/Source/Runtime/VulkanRHI/Private/VulkanRHI.cpp
--- a/Engine/Source/Runtime/VulkanRHI/Private/VulkanRHI.cpp
+++ b/Engine/Source/Runtime/VulkanRHI/Private/VulkanRHI.cpp
@@ -1156,10 +1156,6 @@ void FVulkanDynamicRHI::RHIEndFrame_RenderThread(FRHICommandListImmediate& RHICm
 
 		check(Context.IsImmediate());
 
-#if (RHI_NEW_GPU_PROFILER == 0)
-		Context.GpuProfiler.EndFrame();
-#endif
-
 		bool bTrimMemory = false;
 		Context.FreeUnusedCmdBuffers(bTrimMemory);
 
@@ -1192,10 +1188,6 @@ void FVulkanDynamicRHI::RHIEndFrame_RenderThread(FRHICommandListImmediate& RHICm
 		extern uint32 GVulkanRHIDeletionFrameNumber;
 		++GVulkanRHIDeletionFrameNumber;
 
-#if (RHI_NEW_GPU_PROFILER == 0)
-		Context.GpuProfiler.BeginFrame();
-#endif
-
 		if (GRHISupportsRayTracing)
 		{
 			Context.Device.GetRayTracingCompactionRequestHandler()->Update(Context);
@@ -1205,7 +1197,6 @@ void FVulkanDynamicRHI::RHIEndFrame_RenderThread(FRHICommandListImmediate& RHICm
 
 void FVulkanDynamicRHI::RHIEndFrame(const FRHIEndFrameArgs& Args)
 {
-#if RHI_NEW_GPU_PROFILER
 	// Close the previous frame's timing and start a new one
 	auto Lambda = [this, OldTiming = MoveTemp(CurrentTimingPerQueue)]()
 	{
@@ -1237,7 +1228,7 @@ void FVulkanDynamicRHI::RHIEndFrame(const FRHIEndFrameArgs& Args)
 			break;
 		}
 
-		if(ExternalGPUTime.IsSet())
+		if (ExternalGPUTime.IsSet())
 		{
 			Payload.ExternalGPUTime = ExternalGPUTime.GetValue();
 			ExternalGPUTime.Reset();
@@ -1245,36 +1236,25 @@ void FVulkanDynamicRHI::RHIEndFrame(const FRHIEndFrameArgs& Args)
 
 		// CPU timestamp for the frame boundary event is filled in by the submission thread
 		Payload.EndFrameEvent = UE::RHI::GPUProfiler::FEvent::FFrameBoundary(0, Args.FrameNumber
-#if WITH_RHI_BREADCRUMBS
+		#if WITH_RHI_BREADCRUMBS
 			, (Pipeline != ERHIPipeline::None) ? Args.GPUBreadcrumbs[Pipeline] : nullptr
-#endif
-#if STATS
+		#endif
+		#if STATS
 			, Args.StatsFrame
-#endif
+		#endif
 		);
 	});
-#else
-	{
-		FVulkanPlatformCommandList* Payloads = new FVulkanPlatformCommandList;
-		FVulkanPayload* Payload = new FVulkanPayload(*Device->GetGraphicsQueue());
-		Payload->bEndFrame = true;
-		Payloads->Add(Payload);
-		PendingPayloadsForSubmission.Enqueue(Payloads);
-	}
-#endif
 
 	// Pump the interrupt queue to gather completed events
 	// (required if we're not using an interrupt thread).
 	ProcessInterruptQueueUntil(nullptr);
 }
 
-#if RHI_NEW_GPU_PROFILER
 FVulkanTiming::FVulkanTiming(FVulkanQueue& InQueue)
 	: Queue(InQueue)
 	, EventStream(InQueue.GetProfilerQueue())
 {
 }
-#endif
 
 void FVulkanCommandListContext::RHIEndDrawingViewport(FRHIViewport* ViewportRHI, const FRHIPresentArgs& InPresentArgs)
 {
@@ -1340,8 +1320,7 @@ void FVulkanCommandListContext::RHIEndDrawingViewport(FRHIViewport* ViewportRHI,
 			}
 		#endif
 		}
-		
-#if RHI_NEW_GPU_PROFILER
+
 		if (bSupportsBreadcrumbs)
 		{
 			FlushProfilerStats();
@@ -1351,20 +1330,10 @@ void FVulkanCommandListContext::RHIEndDrawingViewport(FRHIViewport* ViewportRHI,
 			const uint32 IndexInPool = CurrentPool->ReserveQuery(&Event.GPUTimestampTOP);
 			VulkanRHI::vkCmdWriteTimestamp(GetCommandBuffer().GetHandle(), VK_PIPELINE_STAGE_TOP_OF_PIPE_BIT, CurrentPool->GetHandle(), IndexInPool);
 		}
-#else
-		if (IsImmediate())
-		{
-			if (GpuProfiler.IsProfilingGPU())
-			{
-				GpuProfiler.PushEvent(GetNameStr(), Color);
-			}
-		}
-#endif // RHI_NEW_GPU_PROFILER
 	}
 
 	void FVulkanCommandListContext::RHIEndBreadcrumbGPU(FRHIBreadcrumbNode* Breadcrumb)
 	{
-#if RHI_NEW_GPU_PROFILER
 		if (bSupportsBreadcrumbs)
 		{
 			FlushProfilerStats();
@@ -1374,16 +1343,6 @@ void FVulkanCommandListContext::RHIEndDrawingViewport(FRHIViewport* ViewportRHI,
 			const uint32 IndexInPool = CurrentPool->ReserveQuery(&Event.GPUTimestampBOP);
 			VulkanRHI::vkCmdWriteTimestamp(GetCommandBuffer().GetHandle(), VK_PIPELINE_STAGE_BOTTOM_OF_PIPE_BIT, CurrentPool->GetHandle(), IndexInPool);
 		}
-#else
-		//only valid on immediate context currently.  needs to be fixed for parallel rhi execute
-		if (IsImmediate())
-		{
-			if (GpuProfiler.IsProfilingGPU())
-			{
-				GpuProfiler.PopEvent();
-			}
-		}
-#endif // RHI_NEW_GPU_PROFILER
 
 		if (ShouldEmitBreadcrumbs())
 		{
@@ -1650,16 +1609,6 @@ void FVulkanDynamicRHI::RHIFinishExternalComputeWork(VkCommandBuffer InCommandBu
 	ImmediateContext.GetPendingGfxState()->Reset();
 }
 
-void FVulkanDynamicRHI::RHIRegisterWork(uint32 NumPrimitives)
-{
-#if (RHI_NEW_GPU_PROFILER == 0)
-	if (FVulkanPlatform::RegisterGPUWork())
-	{
-		GetDevice()->GetImmediateContext().RegisterGPUWork(NumPrimitives);
-	}
-#endif
-}
-
 void FVulkanDynamicRHI::RHIVerifyResult(VkResult Result, const ANSICHAR* VkFuntion, const ANSICHAR* Filename, uint32 Line)
 {
 	VulkanRHI::VerifyVulkanResult(Result, VkFuntion, Filename, Line);
diff --git a/Engine/Source/Runtime/VulkanRHI/Private/VulkanSubmission.cpp b/Engine/Source/Runtime/VulkanRHI/Private/VulkanSubmission.cpp
--- a/Engine/Source/Runtime/VulkanRHI/Private/VulkanSubmission.cpp
+++ b/Engine/Source/Runtime/VulkanRHI/Private/VulkanSubmission.cpp
@@ -170,7 +170,6 @@ void FVulkanDynamicRHI::InitializeSubmissionPipe()
 		}
 	}
 
-#if RHI_NEW_GPU_PROFILER
 	// Initialize the timing structs in each queue, and the engine GPU profilers
 	{
 		TArray<UE::RHI::GPUProfiler::FQueue> ProfilerQueues;
@@ -186,7 +185,6 @@ void FVulkanDynamicRHI::InitializeSubmissionPipe()
 		UE::RHI::GPUProfiler::InitializeQueues(ProfilerQueues);
 		PendingPayloadsForSubmission.Enqueue(Payloads);
 	}
-#endif
 }
 
 void FVulkanDynamicRHI::ShutdownSubmissionPipe()
@@ -341,31 +339,27 @@ bool FVulkanDynamicRHI::ProcessSubmissionQueue()
 	FVulkanPlatformCommandList* PlatformCmdList = nullptr;
 	while (PendingPayloadsForSubmission.Dequeue(PlatformCmdList))
 	{
-
-#if WITH_RHI_BREADCRUMBS && RHI_NEW_GPU_PROFILER
-		TSharedPtr<FRHIBreadcrumbAllocatorArray> BreadcrumbAllocators{};
+	#if WITH_RHI_BREADCRUMBS
+		TSharedPtr<FRHIBreadcrumbAllocatorArray> BreadcrumbAllocators {};
 		if (PlatformCmdList->BreadcrumbAllocators.Num())
 		{
 			BreadcrumbAllocators = MakeShared<FRHIBreadcrumbAllocatorArray>(MoveTemp(PlatformCmdList->BreadcrumbAllocators));
 		}
+	#endif
 
 		for (FVulkanPayload* Payload : *PlatformCmdList)
 		{
+		#if WITH_RHI_BREADCRUMBS
 			Payload->BreadcrumbRange = PlatformCmdList->BreadcrumbRange;
 			if (BreadcrumbAllocators.IsValid())
 			{
 				check(!Payload->BreadcrumbAllocators.IsValid());
 				Payload->BreadcrumbAllocators = BreadcrumbAllocators;
 			}
+		#endif
 
 			Payload->Queue.EnqueuePayload(Payload);
 		}
-#else
-		for (FVulkanPayload* Payload : *PlatformCmdList)
-		{
-			Payload->Queue.EnqueuePayload(Payload);
-		}
-#endif
 
 		delete PlatformCmdList;
 	}
@@ -728,82 +722,85 @@ void FVulkanDynamicRHI::CompletePayload(FVulkanPayload* Payload)
 
 			for (uint32 QueryIndex = 0; QueryIndex < QueryPool->CurrentQueryCount; ++QueryIndex)
 			{
-					uint64* ResultPtr = QueryPool->QueryResults[QueryIndex];
-					if (ResultPtr)
-					{
-					uint64 QueryResult = 0;
-					for (uint32 ViewIndex = 0; ViewIndex < QueryPool->QueryStride; ++ViewIndex)
+				uint64* ResultPtr = QueryPool->QueryResults[QueryIndex];
+				if (!ResultPtr)
+				{
+					continue;
+				}
+
+				uint64 QueryResult = 0;
+				for (uint32 ViewIndex = 0; ViewIndex < QueryPool->QueryStride; ++ViewIndex)
+				{
+					uint32 ResultIndex = QueryIndex * QueryPool->QueryStride + ViewIndex;
+					if (TempResults[ResultIndex].Availability)
 					{
-						uint32 ResultIndex = QueryIndex * QueryPool->QueryStride + ViewIndex;
-						if (TempResults[ResultIndex].Availability)
+						if (bUseMaxValue)
 						{
-							if (bUseMaxValue)
-							{
-								QueryResult = FMath::Max(QueryResult, TempResults[ResultIndex].Result);
-							}
-							else
-							{
-								QueryResult += TempResults[ResultIndex].Result;
-							}
+							QueryResult = FMath::Max(QueryResult, TempResults[ResultIndex].Result);
 						}
 						else
 						{
-							UE_LOG(LogVulkanRHI, Warning, TEXT("Could not read back query: [PoolType=%d, Index=%d, View=%d]"), (int32)QueryPool->GetPoolType(), QueryIndex, ViewIndex);
+							QueryResult += TempResults[ResultIndex].Result;
 						}
 					}
+					else
+					{
+						UE_LOG(LogVulkanRHI, Warning, TEXT("Could not read back query: [PoolType=%d, Index=%d, View=%d]"), (int32)QueryPool->GetPoolType(), QueryIndex, ViewIndex);
+					}
+				}
 
-#if RHI_NEW_GPU_PROFILER
-						if (QueryPool->GetPoolType() == EVulkanQueryPoolType::Timestamp)
-						{
-							// Convert from GPU timestamp to CPU timestamp (relative to FPlatformTime::Cycles64())
-							checkSlow(Payload->Queue.Timing);
-							const FVulkanTiming* CurrentTiming = Payload->Queue.Timing;
+				if (QueryPool->GetPoolType() == EVulkanQueryPoolType::Timestamp)
+				{
+					// Convert from GPU timestamp to CPU timestamp (relative to FPlatformTime::Cycles64())
+					checkSlow(Payload->Queue.Timing);
+					const FVulkanTiming* CurrentTiming = Payload->Queue.Timing;
 
-							uint64 CPUCalibrationTimestamp = CurrentTiming->CPUTimestamp;
-							uint64 GPUCalibrationTimestamp = CurrentTiming->GPUTimestamp;
+					uint64 CPUCalibrationTimestamp = CurrentTiming->CPUTimestamp;
+					uint64 GPUCalibrationTimestamp = CurrentTiming->GPUTimestamp;
 
-							// NOTE: Adreno GPU calibration values aren't valid, create a fallback
-						if (bBypassTiming || (QueryResult < GPUCalibrationTimestamp) || (GPUCalibrationTimestamp == 0))
+					// NOTE: Adreno GPU calibration values aren't valid, create a fallback
+					if (bBypassTiming || (QueryResult < GPUCalibrationTimestamp) || (GPUCalibrationTimestamp == 0))
+					{
+						if (!bBypassTiming)
+						{
+							static bool bWarnOnce = false;
+							if (!bWarnOnce)
 							{
-								if (!bBypassTiming)
-								{
-									static bool bWarnOnce = false;
-									if (!bWarnOnce)
-									{
-										bWarnOnce = true;
-										UE_LOG(LogVulkanRHI, Warning,
-											TEXT("GPU timestamp can't be moved to CPU timeline.  ")
-											TEXT("Calibration data is invalid (Curr=%lld  Cal=%lld)!  ")
-											TEXT("CPU and GPU profiling data will not line up!"),
-										QueryResult, GPUCalibrationTimestamp);
-									}
-
-									const uint64 CurrentTimestamp = FPlatformTime::Cycles64();
-									ReferenceCPUTimestamp = ((CPUCalibrationTimestamp != 0) && 
-										(CPUCalibrationTimestamp < CurrentTimestamp)) ?
-										CPUCalibrationTimestamp : CurrentTimestamp;
-								ReferenceGPUTimestamp = QueryResult;
-									bBypassTiming = true;
-								}
-
-								CPUCalibrationTimestamp = ReferenceCPUTimestamp;
-								GPUCalibrationTimestamp = ReferenceGPUTimestamp;
-
-								// :todo: Clamp results within the imperfect calibration timestamps (to prevent negative values below)
-							QueryResult = FMath::Max(QueryResult, GPUCalibrationTimestamp);
+								bWarnOnce = true;
+								UE_LOG(LogVulkanRHI, Warning,
+									TEXT("GPU timestamp can't be moved to CPU timeline.  ")
+									TEXT("Calibration data is invalid (Curr=%lld  Cal=%lld)!  ")
+									TEXT("CPU and GPU profiling data will not line up!"),
+								QueryResult, GPUCalibrationTimestamp);
 							}
 
-						const uint64 GPUDelta = QueryResult - GPUCalibrationTimestamp;
-							const uint64 CPUDelta = (GPUDelta * CurrentTiming->CPUFrequency) / CurrentTiming->GPUFrequency;
-							(*ResultPtr) = CPUCalibrationTimestamp + CPUDelta;
-						}
-						else
-#endif
-						{
-						(*ResultPtr) = QueryResult;
+							const uint64 CurrentTimestamp = FPlatformTime::Cycles64();
+
+							ReferenceCPUTimestamp = ((CPUCalibrationTimestamp != 0) && (CPUCalibrationTimestamp < CurrentTimestamp))
+								? CPUCalibrationTimestamp
+								: CurrentTimestamp;
+									
+							ReferenceGPUTimestamp = QueryResult;
+							bBypassTiming = true;
 						}
+
+						CPUCalibrationTimestamp = ReferenceCPUTimestamp;
+						GPUCalibrationTimestamp = ReferenceGPUTimestamp;
+
+						// :todo: Clamp results within the imperfect calibration timestamps (to prevent negative values below)
+						QueryResult = FMath::Max(QueryResult, GPUCalibrationTimestamp);
 					}
+
+					const uint64 GPUDelta = QueryResult - GPUCalibrationTimestamp;
+					const uint64 CPUDelta = (GPUDelta * CurrentTiming->CPUFrequency) / CurrentTiming->GPUFrequency;
+
+					(*ResultPtr) = CPUCalibrationTimestamp + CPUDelta;
+				}
+				else
+				{
+					(*ResultPtr) = QueryResult;
 				}
+			}
 
 			QueryPool->QueryRefs.Empty(QueryPool->GetMaxQueries());
 			Device->ReleaseQueryPool(QueryPool);
@@ -828,12 +825,7 @@ void FVulkanDynamicRHI::CompletePayload(FVulkanPayload* Payload)
 			}
 		}
 
-#if RHI_NEW_GPU_PROFILER
-		const bool bEndFrame = Payload->EndFrameEvent.IsSet();
-#else
-		const bool bEndFrame = Payload->bEndFrame;
-#endif
-		if (bEndFrame)
+		if (Payload->EndFrameEvent.IsSet())
 		{
 			DiagnosticBuffer.CurrentExtendedIndex = (DiagnosticBuffer.CurrentExtendedIndex + 1) % DiagnosticBuffer.ExtendedBreadcrumbAllocators.Num();
 			DiagnosticBuffer.ExtendedBreadcrumbAllocators[DiagnosticBuffer.CurrentExtendedIndex].Reset();
@@ -841,8 +833,6 @@ void FVulkanDynamicRHI::CompletePayload(FVulkanPayload* Payload)
 	}
 #endif // WITH_RHI_BREADCRUMBS
 
-
-#if RHI_NEW_GPU_PROFILER
 	if (!Payload->EventStream.IsEmpty())
 	{
 		check(Payload->Queue.Timing);
@@ -859,25 +849,6 @@ void FVulkanDynamicRHI::CompletePayload(FVulkanPayload* Payload)
 		ReferenceCPUTimestamp = 0;
 		ReferenceGPUTimestamp = 0;
 	}
-#else
-	if (GRHIGlobals.SupportsTimestampRenderQueries && !FVulkanPlatform::HasCustomFrameTiming())
-	{
-		static uint64 BusyCycles = 0;
-		if (Payload->Queue.QueueType == EVulkanQueueType::Graphics)
-		{
-			for (FVulkanCommandBuffer* CommandBuffer : Payload->CommandBuffers)
-			{
-				BusyCycles += CommandBuffer->GetBusyCycles();
-			}
-		}
-		if (Payload->bEndFrame && (BusyCycles > 0))
-		{
-			const double Frequency = double((uint64)((1000.0 * 1000.0 * 1000.0) / Device->GetLimits().timestampPeriod));
-			GRHIGPUFrameTimeHistory.PushFrameCycles(Frequency, BusyCycles);
-			BusyCycles = 0;
-		}
-	}
-#endif
 
 	for (FVulkanSyncPointRef& SyncPoint : Payload->SyncPoints)
 	{
diff --git a/Engine/Source/Runtime/VulkanRHI/Private/VulkanSubmission.h b/Engine/Source/Runtime/VulkanRHI/Private/VulkanSubmission.h
--- a/Engine/Source/Runtime/VulkanRHI/Private/VulkanSubmission.h
+++ b/Engine/Source/Runtime/VulkanRHI/Private/VulkanSubmission.h
@@ -154,9 +154,7 @@ class FVulkanPayload
 public:
 	FVulkanPayload(FVulkanQueue& InQueue)
 		: Queue(InQueue)
-#if RHI_NEW_GPU_PROFILER
 		, EventStream(InQueue.GetProfilerQueue())
-#endif
 	{}
 
 	~FVulkanPayload();
@@ -198,14 +196,10 @@ protected:
 	TSharedPtr<FRHIBreadcrumbAllocatorArray> BreadcrumbAllocators{};
 #endif
 
-#if RHI_NEW_GPU_PROFILER
 	TOptional<FVulkanTiming*> Timing;
 	TOptional<uint64> ExternalGPUTime;
 	TOptional<UE::RHI::GPUProfiler::FEvent::FFrameBoundary> EndFrameEvent;
 	UE::RHI::GPUProfiler::FEventStream EventStream;
-#else
-	bool bEndFrame = false;
-#endif
 
 	// UpdateReservedResources
 	TArray<FVulkanCommitReservedResourceDesc> ReservedResourcesToCommit;
diff --git a/Engine/Source/Runtime/VulkanRHI/Private/VulkanUtil.cpp b/Engine/Source/Runtime/VulkanRHI/Private/VulkanUtil.cpp
--- a/Engine/Source/Runtime/VulkanRHI/Private/VulkanUtil.cpp
+++ b/Engine/Source/Runtime/VulkanRHI/Private/VulkanUtil.cpp
@@ -53,366 +53,6 @@ void SetVulkanResourceName(FVulkanDevice* Device, FVulkanTexture* Texture, const
 #endif
 }
 
-#if (RHI_NEW_GPU_PROFILER == 0)
-
-static FString		EventDeepString("EventTooDeep");
-static const uint32	EventDeepCRC = FCrc::StrCrc32<TCHAR>(*EventDeepString);
-
-/**
- * Initializes the static variables, if necessary.
- */
-void FVulkanGPUTiming::PlatformStaticInitialize(void* UserData)
-{
-	GIsSupported = false;
-
-	// Are the static variables initialized?
-	check( !GAreGlobalsInitialized );
-
-	FVulkanGPUTiming* Caller = (FVulkanGPUTiming*)UserData;
-	if (Caller && Caller->Device && FVulkanPlatform::SupportsTimestampRenderQueries())
-	{
-		const VkPhysicalDeviceLimits& Limits = Caller->Device->GetDeviceProperties().limits;
-		bool bSupportsTimestamps = (Limits.timestampComputeAndGraphics == VK_TRUE);
-		if (!bSupportsTimestamps)
-		{
-			UE_LOG(LogVulkanRHI, Warning, TEXT("Timestamps not supported on Device"));
-			return;
-		}
-		SetTimingFrequency((uint64)((1000.0 * 1000.0 * 1000.0) / Limits.timestampPeriod));
-
-		CalibrateTimers(*Caller->Device);
-		GIsSupported = true;
-	}
-}
-
-void FVulkanGPUTiming::CalibrateTimers(FVulkanDevice& Device)
-{
-	if (Device.GetOptionalExtensions().HasEXTCalibratedTimestamps)
-	{
-		FGPUTimingCalibrationTimestamp CalibrationTimestamp = Device.GetCalibrationTimestamp();
-		SetCalibrationTimestamp(CalibrationTimestamp);
-	}
-}
-
-FVulkanGPUTiming::~FVulkanGPUTiming()
-{
-}
-
-/**
- * Initializes all Vulkan resources and if necessary, the static variables.
- */
-void FVulkanGPUTiming::Initialize(uint32 PoolSize)
-{
-	StaticInitialize(this, PlatformStaticInitialize);
-
-	bIsTiming = false;
-}
-
-/**
- * Releases all Vulkan resources.
- */
-void FVulkanGPUTiming::Release()
-{
-}
-
-/**
- * Start a GPU timing measurement.
- */
-void FVulkanGPUTiming::StartTiming(FVulkanContextCommon* InContext)
-{
-	// Issue a timestamp query for the 'start' time.
-	if (GIsSupported && !bIsTiming)
-	{
-		if (InContext == nullptr)
-		{
-			InContext = Context;
-		}
-
-		// In case we aren't reading queries, remove oldest
-		if (NumPendingQueries >= MaxPendingQueries)
-		{
-			DiscardOldestQuery();
-			if (NumPendingQueries >= MaxPendingQueries)
-			{
-				return;
-			}
-		}
-
-		check(!ActiveQuery);
-		ActiveQuery = new FPendingQuery;
-
-		FVulkanQueryPool* CurrentPool = InContext->GetCurrentTimestampQueryPool();
-		const uint32 IndexInPool = CurrentPool->ReserveQuery(&ActiveQuery->StartResult);
-		VulkanRHI::vkCmdWriteTimestamp(InContext->GetCommandBuffer().GetHandle(), VK_PIPELINE_STAGE_BOTTOM_OF_PIPE_BIT, CurrentPool->GetHandle(), IndexInPool);
-		ActiveQuery->StartSyncPoint = InContext->GetContextSyncPoint();
-
-		bIsTiming = true;
-	}
-}
-
-/**
- * End a GPU timing measurement.
- * The timing for this particular measurement will be resolved at a later time by the GPU.
- */
-void FVulkanGPUTiming::EndTiming(FVulkanContextCommon* InContext)
-{
-	// Issue a timestamp query for the 'end' time.
-	if (GIsSupported && bIsTiming)
-	{
-		if (InContext == nullptr)
-		{
-			InContext = Context;
-		}
-
-		check(ActiveQuery);
-
-		FVulkanQueryPool* CurrentPool = InContext->GetCurrentTimestampQueryPool();
-		const uint32 IndexInPool = CurrentPool->ReserveQuery(&ActiveQuery->EndResult);
-		VulkanRHI::vkCmdWriteTimestamp(InContext->GetCommandBuffer().GetHandle(), VK_PIPELINE_STAGE_BOTTOM_OF_PIPE_BIT, CurrentPool->GetHandle(), IndexInPool);
-		ActiveQuery->EndSyncPoint = InContext->GetContextSyncPoint();
-
-		PendingQueries.Enqueue(ActiveQuery);
-		NumPendingQueries++;
-		ActiveQuery = nullptr;
-
-		bIsTiming = false;
-		bEndTimestampIssued = true;
-	}
-}
-
-void FVulkanGPUTiming::DiscardOldestQuery()
-{
-	FPendingQuery* PendingQuery = nullptr;
-	if (PendingQueries.Peek(PendingQuery))
-	{
-		if (!PendingQuery->StartSyncPoint->IsComplete() || !PendingQuery->EndSyncPoint->IsComplete())
-		{
-			FVulkanDynamicRHI::Get().ProcessInterruptQueueUntil(nullptr);  // leave null, we don't want to force a wait on the SyncPoint
-		}
-
-		if (PendingQuery->StartSyncPoint->IsComplete() && PendingQuery->EndSyncPoint->IsComplete())
-		{
-			PendingQueries.Pop();
-			delete PendingQuery;
-			NumPendingQueries--;
-		}
-	}
-}
-
-/**
- * Retrieves the most recently resolved timing measurement.
- * The unit is the same as for FPlatformTime::Cycles(). Returns 0 if there are no resolved measurements.
- *
- * @return	Value of the most recently resolved timing, or 0 if no measurements have been resolved by the GPU yet.
- */
-uint64 FVulkanGPUTiming::GetTiming(bool bGetCurrentResultsAndBlock)
-{
-	if (GIsSupported)
-	{
-		FPendingQuery* PendingQuery = nullptr;
-		while (PendingQueries.Peek(PendingQuery))
-		{
-			if (bGetCurrentResultsAndBlock && PendingQuery->StartSyncPoint.IsValid() && PendingQuery->EndSyncPoint.IsValid())
-			{
-				SCOPE_CYCLE_COUNTER(STAT_RenderQueryResultTime);
-				{
-					FRenderThreadIdleScope IdleScope(ERenderThreadIdleTypes::WaitingForGPUQuery);
-
-					// CPU wait for query results to be ready.
-					if (!PendingQuery->StartSyncPoint->IsComplete())
-					{
-						FVulkanDynamicRHI::Get().ProcessInterruptQueueUntil(PendingQuery->StartSyncPoint);
-					}
-					if (!PendingQuery->EndSyncPoint->IsComplete())
-					{
-						FVulkanDynamicRHI::Get().ProcessInterruptQueueUntil(PendingQuery->EndSyncPoint);
-					}
-				}
-			}
-
-			if (PendingQuery->EndSyncPoint->IsComplete() && PendingQuery->StartSyncPoint->IsComplete())
-			{
-				if (PendingQuery->EndResult > PendingQuery->StartResult)
-				{
-					// Only keep the most recent result
-					LastTime = PendingQuery->EndResult - PendingQuery->StartResult;
-				}
-
-				PendingQueries.Pop();
-				delete PendingQuery;
-				NumPendingQueries--;
-			}
-			else
-			{
-				break;
-			}
-		}
-
-		return LastTime;
-	}
-
-	return 0;
-}
-
-/** Start this frame of per tracking */
-void FVulkanEventNodeFrame::StartFrame()
-{
-	EventTree.Reset();
-	RootEventTiming.StartTiming();
-}
-
-/** End this frame of per tracking, but do not block yet */
-void FVulkanEventNodeFrame::EndFrame()
-{
-	RootEventTiming.EndTiming();
-}
-
-float FVulkanEventNodeFrame::GetRootTimingResults()
-{
-	double RootResult = 0.0f;
-	if (RootEventTiming.IsSupported())
-	{
-		const uint64 GPUTiming = RootEventTiming.GetTiming(true);
-
-		// In milliseconds
-		RootResult = (double)GPUTiming / (double)RootEventTiming.GetTimingFrequency();
-	}
-
-	return (float)RootResult;
-}
-
-float FVulkanEventNode::GetTiming()
-{
-	float Result = 0;
-
-	if (Timing.IsSupported())
-	{
-		const uint64 GPUTiming = Timing.GetTiming(true);
-		// In milliseconds
-		Result = (double)GPUTiming / (double)Timing.GetTimingFrequency();
-	}
-
-	return Result;
-}
-
-FVulkanGPUProfiler::FVulkanGPUProfiler(FVulkanContextCommon* InContext, FVulkanDevice* InDevice)
-	: bCommandlistSubmitted(false)
-	, Device(InDevice)
-	, CmdContext(InContext)
-	, bBeginFrame(false)
-{
-	BeginFrame();
-}
-
-FVulkanGPUProfiler::~FVulkanGPUProfiler()
-{
-}
-
-void FVulkanGPUProfiler::BeginFrame()
-{
-	bCommandlistSubmitted = false;
-	CurrentEventNode = NULL;
-	check(!bTrackingEvents);
-	check(!CurrentEventNodeFrame); // this should have already been cleaned up and the end of the previous frame
-
-	bBeginFrame = true;
-
-	// latch the bools from the game thread into our private copy
-	bLatchedGProfilingGPU = GTriggerGPUProfile;
-	bLatchedGProfilingGPUHitches = GTriggerGPUHitchProfile;
-	if (bLatchedGProfilingGPUHitches)
-	{
-		bLatchedGProfilingGPU = false; // we do NOT permit an ordinary GPU profile during hitch profiles
-	}
-
-	// if we are starting a hitch profile or this frame is a gpu profile, then save off the state of the draw events
-	if (bLatchedGProfilingGPU || (!bPreviousLatchedGProfilingGPUHitches && bLatchedGProfilingGPUHitches))
-	{
-		bOriginalGEmitDrawEvents = GetEmitDrawEvents();
-	}
-
-	if (bLatchedGProfilingGPU || bLatchedGProfilingGPUHitches)
-	{
-		if (bLatchedGProfilingGPUHitches && GPUHitchDebounce)
-		{
-			// if we are doing hitches and we had a recent hitch, wait to recover
-			// the reasoning is that collecting the hitch report may itself hitch the GPU
-			GPUHitchDebounce--; 
-		}
-		else
-		{
-			SetEmitDrawEvents(true);  // thwart an attempt to turn this off on the game side
-			bTrackingEvents = true;
-			CurrentEventNodeFrame = new FVulkanEventNodeFrame(CmdContext, Device);
-			CurrentEventNodeFrame->StartFrame();
-		}
-	}
-	else if (bPreviousLatchedGProfilingGPUHitches)
-	{
-		// hitch profiler is turning off, clear history and restore draw events
-		GPUHitchEventNodeFrames.Empty();
-		SetEmitDrawEvents(bOriginalGEmitDrawEvents);
-	}
-	bPreviousLatchedGProfilingGPUHitches = bLatchedGProfilingGPUHitches;
-}
-
-void FVulkanGPUProfiler::EndFrameBeforeSubmit()
-{
-	if (GetEmitDrawEvents())
-	{
-		// Finish all open nodes
-		// This is necessary because timestamps must be issued before SubmitDone(), and SubmitDone() happens in RHIEndDrawingViewport instead of RHIEndFrame
-		while (CurrentEventNode)
-		{
-			UE_LOG(LogRHI, Warning, TEXT("POPPING BEFORE SUB"));
-			PopEvent();
-		}
-
-		bCommandlistSubmitted = true;
-	}
-
-	// if we have a frame open, close it now.
-	if (CurrentEventNodeFrame)
-	{
-		CurrentEventNodeFrame->EndFrame();
-	}
-}
-
-void FVulkanGPUProfiler::EndFrame()
-{
-	EndFrameBeforeSubmit();
-
-	check(!bTrackingEvents || bLatchedGProfilingGPU || bLatchedGProfilingGPUHitches);
-	if (bLatchedGProfilingGPU)
-	{
-		if (bTrackingEvents)
-		{
-			CmdContext->FlushCommands(EVulkanFlushFlags::None);
-
-			SetEmitDrawEvents(bOriginalGEmitDrawEvents);
-			UE_LOG(LogRHI, Warning, TEXT(""));
-			UE_LOG(LogRHI, Warning, TEXT(""));
-			check(CurrentEventNodeFrame);
-			CurrentEventNodeFrame->DumpEventTree();
-			GTriggerGPUProfile = false;
-			bLatchedGProfilingGPU = false;
-		}
-	}
-	else if (bLatchedGProfilingGPUHitches)
-	{
-		UE_LOG(LogRHI, Warning, TEXT("GPU hitch tracking not implemented on Vulkan"));
-	}
-	bTrackingEvents = false;
-	if (CurrentEventNodeFrame)
-	{
-		delete CurrentEventNodeFrame;
-		CurrentEventNodeFrame = nullptr;
-	}
-
-	bBeginFrame = false;
-}
-#endif // (RHI_NEW_GPU_PROFILER == 0)
-
 #if NV_AFTERMATH
 struct FShaderPair
 {
diff --git a/Engine/Source/Runtime/VulkanRHI/Public/IVulkanDynamicRHI.h b/Engine/Source/Runtime/VulkanRHI/Public/IVulkanDynamicRHI.h
--- a/Engine/Source/Runtime/VulkanRHI/Public/IVulkanDynamicRHI.h
+++ b/Engine/Source/Runtime/VulkanRHI/Public/IVulkanDynamicRHI.h
@@ -94,8 +94,6 @@ struct IVulkanDynamicRHI : public FDynamicRHI
 	virtual FVulkanRHIAllocationInfo RHIGetAllocationInfo(FRHIBuffer* InBuffer) const = 0;
 
 	virtual void           RHIFinishExternalComputeWork(VkCommandBuffer InCommandBuffer) = 0;
-	virtual void           RHIRegisterWork(uint32 NumPrimitives) = 0;
-
 	virtual void           RHIVerifyResult(VkResult Result, const ANSICHAR* VkFuntion, const ANSICHAR* Filename, uint32 Line) = 0;
 
 	static VULKANRHI_API void AddEnabledInstanceExtensionsAndLayers(TArrayView<const ANSICHAR* const> InInstanceExtensions, TArrayView<const ANSICHAR* const> InInstanceLayers);
@@ -103,6 +101,11 @@ struct IVulkanDynamicRHI : public FDynamicRHI
 
 	// Runs code on SubmissionThread with access to VkQueue.  Useful for plugins. 
 	virtual void RHIRunOnQueue(EVulkanRHIRunOnQueueType QueueType, TFunction<void(VkQueue)>&& CodeToRun, bool bWaitForSubmission) = 0;
+
+	UE_DEPRECATED(5.8, "RHIRegisterWork is deprecated, and there is no replacement. Remove calls to this function.")
+	void RHIRegisterWork(uint32 NumPrimitives)
+	{
+	}
 };
 
 inline IVulkanDynamicRHI* GetIVulkanDynamicRHI()
