VukanRHI:  Always assign bindless handles to UB on creation/update.  Preparing for upcoming features (previous method was not thread safe).

#rnx

diff --git a/Engine/Source/Runtime/VulkanRHI/Private/VulkanBindlessDescriptorManager.cpp b/Engine/Source/Runtime/VulkanRHI/Private/VulkanBindlessDescriptorManager.cpp
--- a/Engine/Source/Runtime/VulkanRHI/Private/VulkanBindlessDescriptorManager.cpp
+++ b/Engine/Source/Runtime/VulkanRHI/Private/VulkanBindlessDescriptorManager.cpp
@@ -45,7 +45,7 @@ static FAutoConsoleVariableRef CVarVulkanBindlessMaxStorageTexelBufferCount(
 	ECVF_ReadOnly
 );
 
-int32 GVulkanBindlessMaxUniformBufferDescriptorCount = 32 * 1024;
+int32 GVulkanBindlessMaxUniformBufferDescriptorCount = 64 * 1024;
 static FAutoConsoleVariableRef CVarVulkanBindlessMaxUniformBufferCount(
 	TEXT("r.Vulkan.Bindless.MaxResourceUniformBufferCount"),
 	GVulkanBindlessMaxUniformBufferDescriptorCount,
diff --git a/Engine/Source/Runtime/VulkanRHI/Private/VulkanCommands.cpp b/Engine/Source/Runtime/VulkanRHI/Private/VulkanCommands.cpp
--- a/Engine/Source/Runtime/VulkanRHI/Private/VulkanCommands.cpp
+++ b/Engine/Source/Runtime/VulkanRHI/Private/VulkanCommands.cpp
@@ -498,7 +498,7 @@ void FVulkanCommandListContext::RHISetUniformBufferDynamicOffset(FUniformBufferS
 			if (Slot == InSlot)
 			{
 				// Uniform views always bind max supported range, so make sure Offset+Range is within buffer allocation
-				check((InOffset + PLATFORM_MAX_UNIFORM_BUFFER_RANGE) <= UniformBuffer->Allocation.Size);
+				check((InOffset + PLATFORM_MAX_UNIFORM_BUFFER_RANGE) <= UniformBuffer->GetAllocation().Size);
 				uint32 DynamicOffset = InOffset + UniformBuffer->GetOffset();
 				PendingGfxState->CurrentState->SetUniformBufferDynamicOffset(Stage, BufferIndex, DynamicOffset);
 
diff --git a/Engine/Source/Runtime/VulkanRHI/Private/VulkanDynamicRHI.h b/Engine/Source/Runtime/VulkanRHI/Private/VulkanDynamicRHI.h
--- a/Engine/Source/Runtime/VulkanRHI/Private/VulkanDynamicRHI.h
+++ b/Engine/Source/Runtime/VulkanRHI/Private/VulkanDynamicRHI.h
@@ -455,8 +455,6 @@ protected:
 	void InternalUpdateTexture2D(FRHICommandListBase& RHICmdList, FRHITexture* TextureRHI, uint32 MipIndex, const struct FUpdateTextureRegion2D& UpdateRegion, uint32 SourcePitch, const uint8* SourceData);
 	void InternalUpdateTexture3D(FRHICommandListBase& RHICmdList, FRHITexture* TextureRHI, uint32 MipIndex, const struct FUpdateTextureRegion3D& UpdateRegion, uint32 SourceRowPitch, uint32 SourceDepthPitch, const uint8* SourceData);
 
-	void UpdateUniformBuffer(FRHICommandListBase& RHICmdList, FVulkanUniformBuffer* UniformBuffer, const void* Contents);
-
 	static void SetupValidationRequests();
 
 	FOptionalVulkanInstanceExtensions OptionalInstanceExtensions;
diff --git a/Engine/Source/Runtime/VulkanRHI/Private/VulkanPipelineState.h b/Engine/Source/Runtime/VulkanRHI/Private/VulkanPipelineState.h
--- a/Engine/Source/Runtime/VulkanRHI/Private/VulkanPipelineState.h
+++ b/Engine/Source/Runtime/VulkanRHI/Private/VulkanPipelineState.h
@@ -97,7 +97,7 @@ public:
 	template<bool bDynamic>
 	inline void SetUniformBuffer(uint8 DescriptorSet, uint32 BindingIndex, const FVulkanUniformBuffer* UniformBuffer)
 	{
-		const VulkanRHI::FVulkanAllocation& Allocation = UniformBuffer->Allocation;
+		const VulkanRHI::FVulkanAllocation& Allocation = UniformBuffer->GetAllocation();
 		VkDeviceSize Range = UniformBuffer->IsUniformView() ? PLATFORM_MAX_UNIFORM_BUFFER_RANGE : UniformBuffer->GetSize();
 
 		if (bDynamic)
diff --git a/Engine/Source/Runtime/VulkanRHI/Private/VulkanResources.h b/Engine/Source/Runtime/VulkanRHI/Private/VulkanResources.h
--- a/Engine/Source/Runtime/VulkanRHI/Private/VulkanResources.h
+++ b/Engine/Source/Runtime/VulkanRHI/Private/VulkanResources.h
@@ -967,39 +967,52 @@ public:
 	void UpdateResourceTable(const FRHIUniformBufferLayout& InLayout, const void* Contents, int32 ResourceNum);
 	void UpdateResourceTable(FRHIResource** Resources, int32 ResourceNum);
 
-	inline VkBuffer GetBufferHandle() const
+	VkBuffer GetBufferHandle() const
 	{
 		return Allocation.GetBufferHandle();
 	}
 
-	inline uint32 GetOffset() const
+	uint32 GetOffset() const
 	{
 		return Allocation.Offset;
 	}
 
-	inline void UpdateAllocation(VulkanRHI::FVulkanAllocation& NewAlloc)
+	const VulkanRHI::FVulkanAllocation& GetAllocation() const
+	{
+		return Allocation;
+	}
+
+	void UpdateAllocation(VulkanRHI::FVulkanAllocation& NewAlloc)
 	{
 		NewAlloc.Swap(Allocation);
+		UpdateBindlessHandle();
 	}
 
-	inline bool IsUniformView() const
+	bool IsUniformView() const
 	{
 		return UniformViewSRV != nullptr;
 	}
+
+	FRHIDescriptorHandle GetBindlessHandle() const
+	{
+		checkfSlow(!BindlessHandle.IsValid() || (CachedDeviceAddress == GetDeviceAddress()), TEXT("Uniform Buffer bindless handle is out of date!"));
+		return BindlessHandle;
+	}
 	
-	FRHIDescriptorHandle GetBindlessHandle();
+	void UpdateBindlessHandle();
 	VkDeviceAddress GetDeviceAddress() const;
 
 	void SetupUniformBufferView();
 
 public:
 	FVulkanDevice* Device;
-	VulkanRHI::FVulkanAllocation Allocation;
-	EUniformBufferUsage Usage;
+	const EUniformBufferUsage Usage;
 
+protected:
+	VulkanRHI::FVulkanAllocation Allocation;
 	FRHIDescriptorHandle BindlessHandle;
 	VkDeviceAddress CachedDeviceAddress = 0;
-	FRHIShaderResourceView* UniformViewSRV;
+	FRHIShaderResourceView* UniformViewSRV = nullptr;
 };
 
 class FVulkanUnorderedAccessView final : public FRHIUnorderedAccessView, public FVulkanLinkedView
diff --git a/Engine/Source/Runtime/VulkanRHI/Private/VulkanUniformBuffer.cpp b/Engine/Source/Runtime/VulkanRHI/Private/VulkanUniformBuffer.cpp
--- a/Engine/Source/Runtime/VulkanRHI/Private/VulkanUniformBuffer.cpp
+++ b/Engine/Source/Runtime/VulkanRHI/Private/VulkanUniformBuffer.cpp
@@ -22,10 +22,6 @@ static FAutoConsoleVariableRef CVarVulkanAllowUniformUpload(
 	ECVF_Default
 );
 
-enum
-{
-	PackedUniformsRingBufferSize = 16 * 1024 * 1024
-};
 
 /*-----------------------------------------------------------------------------
 	Uniform buffer RHI object
@@ -42,12 +38,12 @@ static bool UseTemporaryBuffer(EUniformBufferUsage Usage)
 	return (Usage == UniformBuffer_SingleDraw || Usage == UniformBuffer_SingleFrame);
 }
 
-static void UpdateUniformBufferHelper(FVulkanCommandListContext& Context, FVulkanUniformBuffer* VulkanUniformBuffer, const void* Data, bool bUpdateConstants = true)
+static void UpdateUniformBufferHelper(FVulkanCommandListContext& Context, FVulkanUniformBuffer& VulkanUniformBuffer, const void* Data, bool bUpdateConstants = true)
 {
 	FVulkanCommandBuffer* CmdBuffer = Context.GetActiveCmdBuffer();
 
 	FVulkanDevice& Device = Context.Device;
-	const int32 DataSize = VulkanUniformBuffer->GetLayout().ConstantBufferSize;
+	const int32 DataSize = VulkanUniformBuffer.GetLayout().ConstantBufferSize;
 	const int32 DataAlignment = FMath::Max<uint32>(Device.GetLimits().minUniformBufferOffsetAlignment, 16u);
 
 	VulkanRHI::FVulkanAllocation TempAllocation;
@@ -56,7 +52,7 @@ static void UpdateUniformBufferHelper(FVulkanCommandListContext& Context, FVulka
 	if (bUpdateConstants)
 	{
 		// Update constants as the data is copied
-		UpdateUniformBufferConstants(Device, DestinationData, Data, VulkanUniformBuffer->GetLayoutPtr());
+		UpdateUniformBufferConstants(Device, DestinationData, Data, VulkanUniformBuffer.GetLayoutPtr());
 	}
 	else
 	{
@@ -64,18 +60,9 @@ static void UpdateUniformBufferHelper(FVulkanCommandListContext& Context, FVulka
 		FMemory::Memcpy(DestinationData, Data, DataSize);
 	}
 
-
-	if (UseTemporaryBuffer(VulkanUniformBuffer->Usage))
+	if (UseTemporaryBuffer(VulkanUniformBuffer.Usage))
 	{
-		VulkanUniformBuffer->Allocation.Init(
-			VulkanRHI::EVulkanAllocationEmpty,
-			VulkanRHI::EVulkanAllocationMetaUnknown,
-			TempAllocation.VulkanHandle,
-			DataSize,
-			TempAllocation.Offset,
-			TempAllocation.AllocatorIndex,
-			TempAllocation.AllocationIndex,
-			TempAllocation.HandleId);
+		VulkanUniformBuffer.UpdateAllocation(TempAllocation);
 	}
 	else
 	{
@@ -84,8 +71,8 @@ static void UpdateUniformBufferHelper(FVulkanCommandListContext& Context, FVulka
 		VkBufferCopy Region;
 		Region.size = DataSize;
 		Region.srcOffset = TempAllocation.Offset;
-		Region.dstOffset = VulkanUniformBuffer->GetOffset();
-		VkBuffer UBBuffer = VulkanUniformBuffer->Allocation.GetBufferHandle();
+		Region.dstOffset = VulkanUniformBuffer.GetOffset();
+		VkBuffer UBBuffer = VulkanUniformBuffer.GetBufferHandle();
 
 		VulkanRHI::vkCmdCopyBuffer(CmdBuffer->GetHandle(), TempAllocation.GetBufferHandle(), UBBuffer, 1, &Region);
 	}
@@ -108,7 +95,6 @@ FVulkanUniformBuffer::FVulkanUniformBuffer(FVulkanDevice& InDevice, const FRHIUn
 	: FRHIUniformBuffer(InLayout)
 	, Device(&InDevice)
 	, Usage(InUsage)
-	, UniformViewSRV(nullptr)
 {
 #if VULKAN_ENABLE_AGGRESSIVE_STATS
 	SCOPE_CYCLE_COUNTER(STAT_VulkanUniformBufferCreateTime);
@@ -164,23 +150,22 @@ PRAGMA_ENABLE_DEPRECATION_WARNINGS
 			if (Contents)
 			{
 				FRHICommandListImmediate& RHICmdList = FRHICommandListExecutor::GetImmediateCommandList();
-				FVulkanUniformBuffer* UniformBuffer = this;
 				int32 DataSize = InLayout->ConstantBufferSize;
 
 				// make sure we allocate from RingBuffer on RHIT
 				const bool bCanAllocOnThisThread = RHICmdList.Bypass() || (!IsRunningRHIInSeparateThread() && bInRenderingThread) || bInRHIThread;
 				if (bCanAllocOnThisThread)
 				{
-					UpdateUniformBufferHelper(Device->GetImmediateContext(), UniformBuffer, Contents);
+					UpdateUniformBufferHelper(Device->GetImmediateContext(), *this, Contents);
 				}
 				else
 				{
 					void* CmdListConstantBufferData = RHICmdList.Alloc(DataSize, 16);
 					UpdateUniformBufferConstants(*Device, CmdListConstantBufferData, Contents, InLayout);
 
-					RHICmdList.EnqueueLambda([UniformBuffer, DataSize, CmdListConstantBufferData](FRHICommandList& ExecutingCmdList)
+					RHICmdList.EnqueueLambda([UniformBuffer=this, DataSize, CmdListConstantBufferData](FRHICommandList& ExecutingCmdList)
 					{
-						UpdateUniformBufferHelper(FVulkanCommandListContext::Get(ExecutingCmdList), UniformBuffer, CmdListConstantBufferData, false);
+						UpdateUniformBufferHelper(FVulkanCommandListContext::Get(ExecutingCmdList), *UniformBuffer, CmdListConstantBufferData, false);
 					});
 
 					RHICmdList.RHIThreadFence(true);
@@ -197,8 +182,8 @@ PRAGMA_ENABLE_DEPRECATION_WARNINGS
 				UpdateUniformBufferConstants(*Device, Allocation.GetMappedPointer(Device), Contents, InLayout);
 				Allocation.FlushMappedMemory(Device);
 			}
+			UpdateBindlessHandle();
 		}
-
 	}
 }
 
@@ -233,31 +218,32 @@ void FVulkanUniformBuffer::UpdateResourceTable(FRHIResource** Resources, int32 R
 	}
 }
 
-FRHIDescriptorHandle FVulkanUniformBuffer::GetBindlessHandle()
+void FVulkanUniformBuffer::UpdateBindlessHandle()
 {
-	// :todo-jn: temporary code to refresh as needed, only used by raytracing
-	const VkDeviceAddress CurrentAddress = GetDeviceAddress();
-	if (!BindlessHandle.IsValid() || (CachedDeviceAddress == 0) || (CurrentAddress != CachedDeviceAddress))
+	if (Device->GetBindlessDescriptorManager()->IsSupported())
 	{
-		if (BindlessHandle.IsValid())
+		const VkDeviceAddress CurrentAddress = GetDeviceAddress();
+		if (!BindlessHandle.IsValid() || (CachedDeviceAddress == 0) || (CurrentAddress != CachedDeviceAddress))
 		{
-			Device->GetDeferredDeletionQueue().EnqueueBindlessHandle(BindlessHandle);
-		}
+			if (BindlessHandle.IsValid())
+			{
+				Device->GetDeferredDeletionQueue().EnqueueBindlessHandle(BindlessHandle);
+			}
 
-		BindlessHandle = Device->GetBindlessDescriptorManager()->AllocateDescriptor(VK_DESCRIPTOR_TYPE_UNIFORM_BUFFER, CurrentAddress, GetSize());
+			BindlessHandle = Device->GetBindlessDescriptorManager()->AllocateDescriptor(VK_DESCRIPTOR_TYPE_UNIFORM_BUFFER, CurrentAddress, GetSize());
 
-		CachedDeviceAddress = CurrentAddress;
+			CachedDeviceAddress = CurrentAddress;
+		}
 	}
-
-	return BindlessHandle;
 }
 
 VkDeviceAddress FVulkanUniformBuffer::GetDeviceAddress() const
 {
-	// :todo-jn: there will be more and more churn on this, cache the value
-	VkBufferDeviceAddressInfo BufferInfo;
-	ZeroVulkanStruct(BufferInfo, VK_STRUCTURE_TYPE_BUFFER_DEVICE_ADDRESS_INFO);
-	BufferInfo.buffer = GetBufferHandle();
+	const VkBufferDeviceAddressInfo BufferInfo = {
+		.sType = VK_STRUCTURE_TYPE_BUFFER_DEVICE_ADDRESS_INFO,
+		.pNext = nullptr,
+		.buffer = GetBufferHandle()
+	};
 	VkDeviceAddress BufferAddress = VulkanRHI::vkGetBufferDeviceAddressKHR(Device->GetHandle(), &BufferInfo);
 	BufferAddress += GetOffset();
 	return BufferAddress;
@@ -270,8 +256,10 @@ FUniformBufferRHIRef FVulkanDynamicRHI::RHICreateUniformBuffer(const void* Conte
 	return new FVulkanUniformBuffer(*Device, Layout, Contents, Usage, Validation);
 }
 
-inline void FVulkanDynamicRHI::UpdateUniformBuffer(FRHICommandListBase& RHICmdList, FVulkanUniformBuffer* UniformBuffer, const void* Contents)
+void FVulkanDynamicRHI::RHIUpdateUniformBuffer(FRHICommandListBase& RHICmdList, FRHIUniformBuffer* UniformBufferRHI, const void* Contents)
 {
+	FVulkanUniformBuffer* UniformBuffer = ResourceCast(UniformBufferRHI);
+
 	SCOPE_CYCLE_COUNTER(STAT_VulkanUpdateUniformBuffers);
 
 	const FRHIUniformBufferLayout& Layout = UniformBuffer->GetLayout();
@@ -297,7 +285,7 @@ inline void FVulkanDynamicRHI::UpdateUniformBuffer(FRHICommandListBase& RHICmdLi
 		}
 	}
 
-	bool bRHIBypass = RHICmdList.Bypass();
+	const bool bRHIBypass = RHICmdList.Bypass();
 	if (bRHIBypass)
 	{
 		if (ConstantBufferSize > 0)
@@ -305,7 +293,7 @@ inline void FVulkanDynamicRHI::UpdateUniformBuffer(FRHICommandListBase& RHICmdLi
 			if (bUseUpload || bUseTempBuffer)
 			{			
 				FVulkanCommandListContext& Context = Device->GetImmediateContext();
-				UpdateUniformBufferHelper(Context, UniformBuffer, Contents);
+				UpdateUniformBufferHelper(Context, *UniformBuffer, Contents);
 			}
 			else
 			{
@@ -336,7 +324,7 @@ inline void FVulkanDynamicRHI::UpdateUniformBuffer(FRHICommandListBase& RHICmdLi
 
 			RHICmdList.EnqueueLambda([UniformBuffer, CmdListResources, NumResources, ConstantBufferSize, CmdListConstantBufferData](FRHICommandListBase& ExecutingCmdList)
 			{
-				UpdateUniformBufferHelper(FVulkanCommandListContext::Get(ExecutingCmdList), UniformBuffer, CmdListConstantBufferData);
+				UpdateUniformBufferHelper(FVulkanCommandListContext::Get(ExecutingCmdList), *UniformBuffer, CmdListConstantBufferData);
 				UniformBuffer->UpdateResourceTable(CmdListResources, NumResources);
 			});
 		}
@@ -358,10 +346,3 @@ inline void FVulkanDynamicRHI::UpdateUniformBuffer(FRHICommandListBase& RHICmdLi
 	}
 }
 
-
-void FVulkanDynamicRHI::RHIUpdateUniformBuffer(FRHICommandListBase& RHICmdList, FRHIUniformBuffer* UniformBufferRHI, const void* Contents)
-{
-	FVulkanUniformBuffer* UniformBuffer = ResourceCast(UniformBufferRHI);
-	UpdateUniformBuffer(RHICmdList, UniformBuffer, Contents);
-}
-
