VulkanRHI: Use async tasks to compile shader objects (WIP).

#rnx

#Synced-CL 48106438

diff --git a/Engine/Source/Runtime/VulkanRHI/Private/VulkanResources.h b/Engine/Source/Runtime/VulkanRHI/Private/VulkanResources.h
--- a/Engine/Source/Runtime/VulkanRHI/Private/VulkanResources.h
+++ b/Engine/Source/Runtime/VulkanRHI/Private/VulkanResources.h
@@ -218,6 +218,7 @@ protected:
 
 	VkPipeline						PipelineLibrary = VK_NULL_HANDLE;
 	VkShaderEXT						ShaderObject = VK_NULL_HANDLE;
+	FGraphEventRef					AsyncCompileTask;
 
 protected:
 	class FSpirvContainer
diff --git a/Engine/Source/Runtime/VulkanRHI/Private/VulkanShaderObjectManager.cpp b/Engine/Source/Runtime/VulkanRHI/Private/VulkanShaderObjectManager.cpp
--- a/Engine/Source/Runtime/VulkanRHI/Private/VulkanShaderObjectManager.cpp
+++ b/Engine/Source/Runtime/VulkanRHI/Private/VulkanShaderObjectManager.cpp
@@ -7,6 +7,16 @@
 #include "VulkanBindlessDescriptorManager.h"
 
 
+static int32 GAllowAsyncShaderObjectCompilation = 1;
+FAutoConsoleVariableRef CVarVulkanAllowAsyncShaderObjectCompilation(
+	TEXT("r.Vulkan.AsyncShaderObjectCompilation"),
+	GAllowAsyncShaderObjectCompilation,
+	TEXT("If enabled, ShaderObjects will be compiled in async tasks. (default: 1)"),
+	ECVF_ReadOnly
+);
+
+
+
 FVulkanShaderObjectManager::FVulkanShaderObjectManager(FVulkanDevice& InDevice)
 	: Device(InDevice)
 {
@@ -36,7 +46,22 @@ VkShaderStageFlags FVulkanShaderObjectManager::GetPossibleNextStages(VkShaderSta
 	return 0;
 }
 
-VkShaderEXT FVulkanShaderObjectManager::CreateShaderObject(FVulkanShader& VulkanShader)
+void FVulkanShaderObjectManager::CreateShaderObject(FVulkanShader& VulkanShader)
+{
+	if (GAllowAsyncShaderObjectCompilation)
+	{
+		VulkanShader.AsyncCompileTask = FFunctionGraphTask::CreateAndDispatchWhenReady([ShaderObjectManager = this, &VulkanShader]()
+		{
+			ShaderObjectManager->InternalCreateShaderObject(VulkanShader);
+		});
+	}
+	else
+	{
+		InternalCreateShaderObject(VulkanShader);
+	}
+}
+
+void FVulkanShaderObjectManager::InternalCreateShaderObject(FVulkanShader& VulkanShader)
 {
 	checkf(VulkanShader.UsesBindless(), TEXT("Shader Objects currently only supported in bindless."));
 
@@ -88,15 +113,13 @@ VkShaderEXT FVulkanShaderObjectManager::CreateShaderObject(FVulkanShader& Vulkan
 		}
 	}
 
-	VkShaderEXT ShaderObjectHandle = VK_NULL_HANDLE;
-	VERIFYVULKANRESULT(VulkanRHI::vkCreateShadersEXT(Device.GetHandle(), 1, &ShaderCreateInfo, VULKAN_CPU_ALLOCATOR, &ShaderObjectHandle));
+	check(VulkanShader.ShaderObject == VK_NULL_HANDLE);
+	VERIFYVULKANRESULT(VulkanRHI::vkCreateShadersEXT(Device.GetHandle(), 1, &ShaderCreateInfo, VULKAN_CPU_ALLOCATOR, &VulkanShader.ShaderObject));
 
 	if (!VulkanShader.GetCodeHeader().DebugName.IsEmpty())
 	{
-		VULKAN_SET_DEBUG_NAME(Device, VK_OBJECT_TYPE_SHADER_EXT, ShaderObjectHandle, TEXT("%s"), *VulkanShader.GetCodeHeader().DebugName);
+		VULKAN_SET_DEBUG_NAME(Device, VK_OBJECT_TYPE_SHADER_EXT, VulkanShader.ShaderObject, TEXT("%s"), *VulkanShader.GetCodeHeader().DebugName);
 	}
-
-	return ShaderObjectHandle;
 }
 
 bool FVulkanShaderObjectManager::CreateLinkedShaderObjects(FVulkanGraphicsPipelineState* PSO, FVulkanShader* Shaders[ShaderStage::NumGraphicsStages], FGraphicsPipelineStateInitializer::EPSOPrecacheCompileType PSOCompileType)
@@ -115,10 +138,25 @@ void FVulkanShaderObjectManager::BindGraphicsShaderObjects(VkCommandBuffer Comma
 	uint32 StageCount = 0;
 	auto AddStage = [&StageCount, &ShaderStages, &ShaderObjectHandles, &VulkanShaders](ShaderStage::EStage Stage, VkShaderStageFlagBits StageFlag)
 	{
-		FVulkanShader* VulkanShader = VulkanShaders[Stage];
 		ShaderStages[StageCount] = StageFlag;
-		ShaderObjectHandles[StageCount] = VulkanShader ? VulkanShader->ShaderObject : VK_NULL_HANDLE;
-		check(!VulkanShader || (Stage == ShaderStage::GetStageForFrequency(VulkanShader->Frequency)));
+
+		FVulkanShader* VulkanShader = VulkanShaders[Stage];
+		if (VulkanShader)
+		{
+			if (VulkanShader->AsyncCompileTask.IsValid() && !VulkanShader->AsyncCompileTask->IsComplete())
+			{
+				VulkanShader->AsyncCompileTask->Wait();
+			}
+			ShaderObjectHandles[StageCount] = VulkanShader->ShaderObject;
+			checkf(Stage == ShaderStage::GetStageForFrequency(VulkanShader->Frequency), 
+				TEXT("Shader Frequency (%s) does not match expected Vulkan shader stage (%d)!"), 
+				GetShaderFrequencyString(VulkanShader->Frequency), Stage);
+		}
+		else
+		{
+			ShaderObjectHandles[StageCount] = VK_NULL_HANDLE;
+		}
+
 		StageCount++;
 	};
 
@@ -146,7 +184,14 @@ void FVulkanShaderObjectManager::BindGraphicsShaderObjects(VkCommandBuffer Comma
 
 void FVulkanShaderObjectManager::BindComputeShaderObject(VkCommandBuffer CommandBufferHandle, FVulkanShader* VulkanShader)
 {
-	check(VulkanShader && VulkanShader->ShaderObject);
+	check(VulkanShader);
+
+	if (VulkanShader->AsyncCompileTask.IsValid() && !VulkanShader->AsyncCompileTask->IsComplete())
+	{
+		VulkanShader->AsyncCompileTask->Wait();
+	}
+
+	check(VulkanShader->ShaderObject);
 
 	static const VkShaderStageFlagBits ShaderStage = VK_SHADER_STAGE_COMPUTE_BIT;
 	const VkShaderEXT ShaderObjectHandle = VulkanShader->ShaderObject;
diff --git a/Engine/Source/Runtime/VulkanRHI/Private/VulkanShaderObjectManager.h b/Engine/Source/Runtime/VulkanRHI/Private/VulkanShaderObjectManager.h
--- a/Engine/Source/Runtime/VulkanRHI/Private/VulkanShaderObjectManager.h
+++ b/Engine/Source/Runtime/VulkanRHI/Private/VulkanShaderObjectManager.h
@@ -16,7 +16,7 @@ class FVulkanShaderObjectManager
 public:
 	FVulkanShaderObjectManager(FVulkanDevice& InDevice);
 
-	VkShaderEXT CreateShaderObject(FVulkanShader& VulkanShader);
+	void CreateShaderObject(FVulkanShader& VulkanShader);
 	bool CreateLinkedShaderObjects(FVulkanGraphicsPipelineState* PSO, FVulkanShader* Shaders[ShaderStage::NumGraphicsStages], FGraphicsPipelineStateInitializer::EPSOPrecacheCompileType PSOCompileType);
 
 	void BindGraphicsShaderObjects(VkCommandBuffer CommandBufferHandle, TArrayView<FVulkanShader*> VulkanShaders);
@@ -24,6 +24,7 @@ public:
 
 private:
 	VkShaderStageFlags GetPossibleNextStages(VkShaderStageFlagBits Stage);
+	void InternalCreateShaderObject(FVulkanShader& VulkanShader);
 
 	FVulkanDevice& Device;
 };
\ No newline at end of file
diff --git a/Engine/Source/Runtime/VulkanRHI/Private/VulkanShaders.cpp b/Engine/Source/Runtime/VulkanRHI/Private/VulkanShaders.cpp
--- a/Engine/Source/Runtime/VulkanRHI/Private/VulkanShaders.cpp
+++ b/Engine/Source/Runtime/VulkanRHI/Private/VulkanShaders.cpp
@@ -249,7 +249,7 @@ ShaderType* FVulkanShaderFactory::CreateShader(TArrayView<const uint8> Code, FVu
 
 				if (Device->SupportsShaderObjects())
 				{
-					RetShader->ShaderObject = Device->GetShaderObjectManager()->CreateShaderObject(*RetShader);
+					Device->GetShaderObjectManager()->CreateShaderObject(*RetShader);
 				}
 				else if (Device->SupportsGraphicPipelineLibraries())
 				{
@@ -738,6 +738,11 @@ FVulkanShader::~FVulkanShader()
 	PurgeShaderModules();
 	Device->GetShaderFactory().OnDeleteShader(*this);
 
+	if (AsyncCompileTask.IsValid() && !AsyncCompileTask->IsComplete())
+	{
+		AsyncCompileTask->Wait();
+	}
+
 	if (ShaderObject)
 	{
 		Device->GetDeferredDeletionQueue().EnqueueFunc([Handle=ShaderObject](FVulkanDevice& Device) 
