Shuffling Vulkan state around to allow for pipeline independent setters.
This is a prerequisite for removing the RHISetShaderXXX methods.

#Synced-CL 48184602

diff --git a/Engine/Source/Runtime/VulkanRHI/Private/VulkanCommands.cpp b/Engine/Source/Runtime/VulkanRHI/Private/VulkanCommands.cpp
--- a/Engine/Source/Runtime/VulkanRHI/Private/VulkanCommands.cpp
+++ b/Engine/Source/Runtime/VulkanRHI/Private/VulkanCommands.cpp
@@ -121,22 +121,22 @@ void FVulkanCommandListContext::RHISetStreamSource(uint32 StreamIndex, FRHIBuffe
 	}
 }
 
-template <class PendingStateType>
 struct FVulkanResourceBinder
 {
 	FVulkanCommandListContext& Context;
 	const EShaderFrequency Frequency;
 	const ShaderStage::EStage Stage;
-	PendingStateType* PendingState;
+	FVulkanCommonPipelineDescriptorState& State;
 
 	// Even if no resource are bound when bindless is enabled, we still need to process the Resource Table to go through proper validation
 	const bool bBindless;
 
+	template <class PendingStateType>
 	FVulkanResourceBinder(FVulkanCommandListContext& InContext, EShaderFrequency InFrequency, PendingStateType* InPendingState, bool InUsesBindless)
 		: Context(InContext)
 		, Frequency(InFrequency)
 		, Stage((InFrequency == SF_Compute) ? ShaderStage::Compute : ShaderStage::GetStageForFrequency(InFrequency))
-		, PendingState(InPendingState)
+		, State(InPendingState->GetCurrentState())
 		, bBindless(InUsesBindless)
 	{
 	}
@@ -150,7 +150,7 @@ struct FVulkanResourceBinder
 
 		if (!bBindless)
 		{
-			PendingState->SetUAVForUBResource(Stage, Index, ResourceCast(UAV));
+			State.SetUAV(Frequency == SF_Compute, Stage, Index, ResourceCast(UAV));
 		}
 	}
 
@@ -158,7 +158,7 @@ struct FVulkanResourceBinder
 	{
 		if (!bBindless)
 		{
-			PendingState->SetSRVForUBResource(Stage, Index, ResourceCast(SRV));
+			State.SetSRV(Frequency == SF_Compute, Stage, Index, ResourceCast(SRV));
 		}
 	}
 
@@ -169,7 +169,7 @@ struct FVulkanResourceBinder
 			FVulkanTexture* VulkanTexture = ResourceCast(TextureRHI);
 			const ERHIAccess RHIAccess = (Frequency == SF_Compute) ? ERHIAccess::SRVCompute : ERHIAccess::SRVGraphics;
 			const VkImageLayout ExpectedLayout = FVulkanPipelineBarrier::GetDefaultLayout(*VulkanTexture, RHIAccess);
-			PendingState->SetTextureForUBResource(Stage, Index, VulkanTexture, ExpectedLayout);
+			State.SetTexture(Stage, Index, VulkanTexture, ExpectedLayout);
 		}
 	}
 
@@ -177,7 +177,7 @@ struct FVulkanResourceBinder
 	{
 		if (!bBindless)
 		{
-			PendingState->SetSamplerStateForUBResource(Stage, Index, ResourceCast(Sampler));
+			State.SetSamplerState(Stage, Index, ResourceCast(Sampler));
 		}
 	}
 
@@ -319,8 +319,7 @@ void FVulkanCommandListContext::RHISetUAVParameter(FRHIPixelShader* PixelShaderR
 {
 	if (UAVRHI)
 	{
-		FVulkanUnorderedAccessView* UAV = ResourceCast(UAVRHI);
-		PendingGfxState->SetUAVForStage(ShaderStage::Pixel, UAVIndex, UAV);
+		PendingGfxState->GetCurrentState().SetUAV(false, ShaderStage::Pixel, UAVIndex, ResourceCast(UAVRHI));
 	}
 }
 
@@ -329,36 +328,29 @@ void FVulkanCommandListContext::RHISetUAVParameter(FRHIComputeShader* ComputeSha
 	if (UAVRHI)
 	{
 		check(PendingComputeState->GetCurrentShader() == ResourceCast(ComputeShaderRHI));
-
-		FVulkanUnorderedAccessView* UAV = ResourceCast(UAVRHI);
-		PendingComputeState->SetUAVForStage(UAVIndex, UAV);
+		PendingComputeState->GetCurrentState().SetUAV(true, ShaderStage::Compute, UAVIndex, ResourceCast(UAVRHI));
 	}
 }
 
-void FVulkanCommandListContext::RHISetUAVParameter(FRHIComputeShader* ComputeShaderRHI,uint32 UAVIndex, FRHIUnorderedAccessView* UAVRHI, uint32 InitialCount)
-{
-	ensure(0);
-}
-
-
 void FVulkanCommandListContext::RHISetShaderTexture(FRHIGraphicsShader* ShaderRHI, uint32 TextureIndex, FRHITexture* NewTextureRHI)
 {
 	FVulkanTexture* VulkanTexture = ResourceCast(NewTextureRHI);
 	const VkImageLayout ExpectedLayout = FVulkanPipelineBarrier::GetDefaultLayout(*VulkanTexture, ERHIAccess::SRVGraphics);
 
-	ShaderStage::EStage Stage = GetAndVerifyShaderStage(ShaderRHI, PendingGfxState);
-	PendingGfxState->SetTextureForStage(Stage, TextureIndex, VulkanTexture, ExpectedLayout);
+	const ShaderStage::EStage Stage = GetAndVerifyShaderStage(ShaderRHI, PendingGfxState);
+	PendingGfxState->GetCurrentState().SetTexture(Stage, TextureIndex, VulkanTexture, ExpectedLayout);
+
 	NewTextureRHI->SetLastRenderTime((float)FPlatformTime::Seconds());
 }
 
 void FVulkanCommandListContext::RHISetShaderTexture(FRHIComputeShader* ComputeShaderRHI, uint32 TextureIndex, FRHITexture* NewTextureRHI)
 {
-	FVulkanComputeShader* ComputeShader = ResourceCast(ComputeShaderRHI);
-	check(PendingComputeState->GetCurrentShader() == ComputeShader);
+	check(PendingComputeState->GetCurrentShader() == ResourceCast(ComputeShaderRHI));
 
 	FVulkanTexture* VulkanTexture = ResourceCast(NewTextureRHI);
 	const VkImageLayout ExpectedLayout = FVulkanPipelineBarrier::GetDefaultLayout(*VulkanTexture, ERHIAccess::SRVCompute);
-	PendingComputeState->SetTextureForStage(TextureIndex, VulkanTexture, ExpectedLayout);
+	PendingComputeState->GetCurrentState().SetTexture(ShaderStage::Compute, TextureIndex, VulkanTexture, ExpectedLayout);
+
 	NewTextureRHI->SetLastRenderTime((float)FPlatformTime::Seconds());
 }
 
@@ -366,9 +358,8 @@ void FVulkanCommandListContext::RHISetShaderResourceViewParameter(FRHIGraphicsSh
 {
 	if (SRVRHI)
 	{
-		ShaderStage::EStage Stage = GetAndVerifyShaderStage(ShaderRHI, PendingGfxState);
-		FVulkanShaderResourceView* SRV = ResourceCast(SRVRHI);
-		PendingGfxState->SetSRVForStage(Stage, TextureIndex, SRV);
+		const ShaderStage::EStage Stage = GetAndVerifyShaderStage(ShaderRHI, PendingGfxState);
+		PendingGfxState->GetCurrentState().SetSRV(false, Stage, TextureIndex, ResourceCast(SRVRHI));
 	}
 }
 
@@ -377,40 +368,32 @@ void FVulkanCommandListContext::RHISetShaderResourceViewParameter(FRHIComputeSha
 	if (SRVRHI)
 	{
 		check(PendingComputeState->GetCurrentShader() == ResourceCast(ComputeShaderRHI));
-
-		FVulkanShaderResourceView* SRV = ResourceCast(SRVRHI);
-		PendingComputeState->SetSRVForStage(TextureIndex, SRV);
+		PendingComputeState->GetCurrentState().SetSRV(true, ShaderStage::Compute, TextureIndex, ResourceCast(SRVRHI));
 	}
 }
 
 void FVulkanCommandListContext::RHISetShaderSampler(FRHIGraphicsShader* ShaderRHI, uint32 SamplerIndex, FRHISamplerState* NewStateRHI)
 {
-	ShaderStage::EStage Stage = GetAndVerifyShaderStage(ShaderRHI, PendingGfxState);
-	FVulkanSamplerState* Sampler = ResourceCast(NewStateRHI);
-	PendingGfxState->SetSamplerStateForStage(Stage, SamplerIndex, Sampler);
+	const ShaderStage::EStage Stage = GetAndVerifyShaderStage(ShaderRHI, PendingGfxState);
+	PendingGfxState->GetCurrentState().SetSamplerState(Stage, SamplerIndex, ResourceCast(NewStateRHI));
 }
 
 void FVulkanCommandListContext::RHISetShaderSampler(FRHIComputeShader* ComputeShaderRHI, uint32 SamplerIndex, FRHISamplerState* NewStateRHI)
 {
-	FVulkanComputeShader* ComputeShader = ResourceCast(ComputeShaderRHI);
-	check(PendingComputeState->GetCurrentShader() == ComputeShader);
-
-	FVulkanSamplerState* Sampler = ResourceCast(NewStateRHI);
-	PendingComputeState->SetSamplerStateForStage(SamplerIndex, Sampler);
+	check(PendingComputeState->GetCurrentShader() == ResourceCast(ComputeShaderRHI));
+	PendingComputeState->GetCurrentState().SetSamplerState(ShaderStage::Compute, SamplerIndex, ResourceCast(NewStateRHI));
 }
 
 void FVulkanCommandListContext::RHISetShaderParameter(FRHIGraphicsShader* ShaderRHI, uint32 BufferIndex, uint32 BaseIndex, uint32 NumBytes, const void* NewValue)
 {
-	ShaderStage::EStage Stage = GetAndVerifyShaderStage(ShaderRHI, PendingGfxState);
-	PendingGfxState->SetPackedGlobalShaderParameter(Stage, BufferIndex, BaseIndex, NumBytes, NewValue);
+	const ShaderStage::EStage Stage = GetAndVerifyShaderStage(ShaderRHI, PendingGfxState);
+	PendingGfxState->GetCurrentState().GetUniformBuffers(Stage).SetParameter(BaseIndex, NumBytes, NewValue);
 }
 
 void FVulkanCommandListContext::RHISetShaderParameter(FRHIComputeShader* ComputeShaderRHI,uint32 BufferIndex, uint32 BaseIndex, uint32 NumBytes, const void* NewValue)
 {
-	FVulkanComputeShader* ComputeShader = ResourceCast(ComputeShaderRHI);
-	check(PendingComputeState->GetCurrentShader() == ComputeShader);
-
-	PendingComputeState->SetPackedGlobalShaderParameter(BufferIndex, BaseIndex, NumBytes, NewValue);
+	check(PendingComputeState->GetCurrentShader() == ResourceCast(ComputeShaderRHI));
+	PendingComputeState->GetCurrentState().GetUniformBuffers(ShaderStage::Compute).SetParameter(BaseIndex, NumBytes, NewValue);
 }
 
 void FVulkanCommandListContext::RHISetShaderParameters(FRHIGraphicsShader* Shader, TConstArrayView<uint8> InParametersData, TConstArrayView<FRHIShaderParameter> InParameters, TConstArrayView<FRHIShaderParameterResource> InResourceParameters, TConstArrayView<FRHIShaderParameterResource> InBindlessParameters)
@@ -508,17 +491,20 @@ void FVulkanCommandListContext::RHISetUniformBufferDynamicOffset(FUniformBufferS
 	}
 }
 
-void FVulkanCommandListContext::RHISetShaderUniformBuffer(FRHIGraphicsShader* ShaderRHI, uint32 BufferIndex, FRHIUniformBuffer* BufferRHI)
+static void BindUniformBuffer(
+	FVulkanCommandListContext& Context,
+	FVulkanCommonPipelineDescriptorState& State,
+	FVulkanShader* Shader,
+	EShaderFrequency Frequency,
+	ShaderStage::EStage Stage,
+	uint32 BufferIndex,
+	FVulkanUniformBuffer* UniformBuffer
+)
 {
 #if VULKAN_ENABLE_AGGRESSIVE_STATS
 	SCOPE_CYCLE_COUNTER(STAT_VulkanSetUniformBufferTime);
 #endif
 
-	FVulkanShader* Shader = nullptr;
-	const ShaderStage::EStage Stage = GetAndVerifyShaderStageAndVulkanShader(ShaderRHI, PendingGfxState, Shader);
-	check(Shader->GetShaderKey() == PendingGfxState->GetCurrentShaderKey(Stage));
-
-	FVulkanUniformBuffer* UniformBuffer = ResourceCast(BufferRHI);
 	const FVulkanShaderHeader& CodeHeader = Shader->GetCodeHeader();
 	checkfSlow(!CodeHeader.UniformBufferInfos[BufferIndex].LayoutHash || (CodeHeader.UniformBufferInfos[BufferIndex].LayoutHash == UniformBuffer->GetLayout().GetHash()),
 		TEXT("Mismatched UB layout! Got hash 0x%x, expected 0x%x!"), UniformBuffer->GetLayout().GetHash(), CodeHeader.UniformBufferInfos[BufferIndex].LayoutHash);
@@ -532,17 +518,17 @@ void FVulkanCommandListContext::RHISetShaderUniformBuffer(FRHIGraphicsShader* Sh
 	if (BufferIndex < CodeHeader.NumBoundUniformBuffers)
 	{
 		checkSlow(UniformBuffer->GetLayout().ConstantBufferSize > 0);
-		
-		const VkDescriptorType DescriptorType = PendingGfxState->CurrentState->GetDescriptorType(Stage, BufferIndex);
+
+		const VkDescriptorType DescriptorType = State.GetDescriptorType(Stage, BufferIndex);
 
 		if (DescriptorType == VK_DESCRIPTOR_TYPE_UNIFORM_BUFFER_DYNAMIC)
 		{
-			PendingGfxState->SetUniformBuffer<true>(Stage, BufferIndex, UniformBuffer);
+			State.SetUniformBuffer<true>(Stage, BufferIndex, UniformBuffer);
 		}
 		else
 		{
 			check(DescriptorType == VK_DESCRIPTOR_TYPE_UNIFORM_BUFFER);
-			PendingGfxState->SetUniformBuffer<false>(Stage, BufferIndex, UniformBuffer);
+			State.SetUniformBuffer<false>(Stage, BufferIndex, UniformBuffer);
 		}
 
 		bHasResources = (CodeHeader.UniformBufferInfos[BufferIndex].bHasResources != 0);
@@ -555,64 +541,43 @@ void FVulkanCommandListContext::RHISetShaderUniformBuffer(FRHIGraphicsShader* Sh
 
 	if (bHasResources)
 	{
-		checkSlow(Shader->Frequency < SF_NumStandardFrequencies);
-		check(BufferIndex < MAX_UNIFORM_BUFFERS_PER_SHADER_STAGE);
-		BoundUniformBuffers[Shader->Frequency][BufferIndex] = UniformBuffer;
-		DirtyUniformBuffers[Shader->Frequency] |= (1 << BufferIndex);
+		checkSlow(Frequency < SF_NumStandardFrequencies);
+		check(BufferIndex < FVulkanCommandListContext::MAX_UNIFORM_BUFFERS_PER_SHADER_STAGE);
+		Context.BoundUniformBuffers[Frequency][BufferIndex] = UniformBuffer;
+		Context.DirtyUniformBuffers[Frequency] |= (1 << BufferIndex);
 	}
 }
 
-void FVulkanCommandListContext::RHISetShaderUniformBuffer(FRHIComputeShader* ComputeShaderRHI, uint32 BufferIndex, FRHIUniformBuffer* BufferRHI)
+void FVulkanCommandListContext::RHISetShaderUniformBuffer(FRHIGraphicsShader* ShaderRHI, uint32 BufferIndex, FRHIUniformBuffer* BufferRHI)
 {
-	FVulkanComputeShader* ComputeShader = ResourceCast(ComputeShaderRHI);
-	check(PendingComputeState->GetCurrentShader() == ComputeShader);
-
-#if VULKAN_ENABLE_AGGRESSIVE_STATS
-	SCOPE_CYCLE_COUNTER(STAT_VulkanSetUniformBufferTime);
-#endif
-	FVulkanComputePipelineDescriptorState& State = *PendingComputeState->CurrentState;
-
-	// Walk through all resources to set all appropriate states
-	FVulkanComputeShader* Shader = ResourceCast(ComputeShaderRHI);
-	FVulkanUniformBuffer* UniformBuffer = ResourceCast(BufferRHI);
-
-	const FVulkanShaderHeader& CodeHeader = Shader->GetCodeHeader();
-	checkfSlow(!CodeHeader.UniformBufferInfos[BufferIndex].LayoutHash || (CodeHeader.UniformBufferInfos[BufferIndex].LayoutHash == UniformBuffer->GetLayout().GetHash()), 
-		TEXT("Mismatched UB layout! Got hash 0x%x, expected 0x%x!"), UniformBuffer->GetLayout().GetHash(), CodeHeader.UniformBufferInfos[BufferIndex].LayoutHash);
-
-	// Uniform Buffers
-	bool bHasResources = false;
-	if (BufferIndex < CodeHeader.NumBoundUniformBuffers)
-	{
-		checkSlow(UniformBuffer->GetLayout().ConstantBufferSize > 0);
-
-		const VkDescriptorType DescriptorType = State.GetDescriptorType(ShaderStage::Compute, BufferIndex);
+	FVulkanShader* Shader = nullptr;
+	const ShaderStage::EStage Stage = GetAndVerifyShaderStageAndVulkanShader(ShaderRHI, PendingGfxState, Shader);
+	check(Shader->GetShaderKey() == PendingGfxState->GetCurrentShaderKey(Stage));
 
-		if (DescriptorType == VK_DESCRIPTOR_TYPE_UNIFORM_BUFFER_DYNAMIC)
-		{
-			State.SetUniformBuffer<true>(ShaderStage::Compute, BufferIndex, UniformBuffer);
-		}
-		else
-		{
-			check(DescriptorType == VK_DESCRIPTOR_TYPE_UNIFORM_BUFFER);
-			State.SetUniformBuffer<false>(ShaderStage::Compute, BufferIndex, UniformBuffer);
-		}
+	BindUniformBuffer(
+		*this,
+		PendingGfxState->GetCurrentState(),
+		Shader,
+		Shader->Frequency,
+		Stage,
+		BufferIndex,
+		ResourceCast(BufferRHI)
+	);
+}
 
-		bHasResources = (CodeHeader.UniformBufferInfos[BufferIndex].bHasResources != 0);
-	}
-	else
-	{
-		// If the buffer has no bindings, then it is as resource only ub
-		bHasResources = true;
-	}
+void FVulkanCommandListContext::RHISetShaderUniformBuffer(FRHIComputeShader* ComputeShaderRHI, uint32 BufferIndex, FRHIUniformBuffer* BufferRHI)
+{
+	check(PendingComputeState->GetCurrentShader() == ResourceCast(ComputeShaderRHI));
 
-	if (bHasResources)
-	{
-		checkSlow(ComputeShaderRHI->GetFrequency() == SF_Compute);
-		check(BufferIndex < MAX_UNIFORM_BUFFERS_PER_SHADER_STAGE);
-		BoundUniformBuffers[SF_Compute][BufferIndex] = UniformBuffer;
-		DirtyUniformBuffers[SF_Compute] |= (1 << BufferIndex);
-	}
+	BindUniformBuffer(
+		*this,
+		PendingComputeState->GetCurrentState(),
+		ResourceCast(ComputeShaderRHI),
+		SF_Compute,
+		ShaderStage::Compute,
+		BufferIndex,
+		ResourceCast(BufferRHI)
+	);
 }
 
 void FVulkanCommandListContext::RHISetStencilRef(uint32 StencilRef)
diff --git a/Engine/Source/Runtime/VulkanRHI/Private/VulkanContext.h b/Engine/Source/Runtime/VulkanRHI/Private/VulkanContext.h
--- a/Engine/Source/Runtime/VulkanRHI/Private/VulkanContext.h
+++ b/Engine/Source/Runtime/VulkanRHI/Private/VulkanContext.h
@@ -275,7 +275,6 @@ public:
 	void RHISetShaderSampler(FRHIGraphicsShader* Shader, uint32 SamplerIndex, FRHISamplerState* NewState);
 	void RHISetUAVParameter(FRHIPixelShader* PixelShader, uint32 UAVIndex, FRHIUnorderedAccessView* UAV);
 	void RHISetUAVParameter(FRHIComputeShader* ComputeShader, uint32 UAVIndex, FRHIUnorderedAccessView* UAV);
-	void RHISetUAVParameter(FRHIComputeShader* ComputeShader, uint32 UAVIndex, FRHIUnorderedAccessView* UAV, uint32 InitialCount);
 	void RHISetShaderResourceViewParameter(FRHIGraphicsShader* Shader, uint32 SamplerIndex, FRHIShaderResourceView* SRV);
 	void RHISetShaderResourceViewParameter(FRHIComputeShader* ComputeShader, uint32 SamplerIndex, FRHIShaderResourceView* SRV);
 	virtual void RHISetStaticUniformBuffers(const FUniformBufferStaticBindings& InUniformBuffers) final override;
@@ -409,7 +408,6 @@ public:
 
 	void ApplyShaderBindingLayout(VkShaderStageFlags ShaderStageFlags);
 
-protected:
 	void ExtractDepthStencilLayouts(const FRHIRenderPassInfo& InInfo, VkImageLayout& OutDepthLayout, VkImageLayout& OutStencilLayout);
 	void FillDynamicRenderingInfo(const FRHIRenderPassInfo& RenderPassInfo, FVulkanDynamicRenderingInfo& OutRenderingInfo);
 
diff --git a/Engine/Source/Runtime/VulkanRHI/Private/VulkanPendingState.cpp b/Engine/Source/Runtime/VulkanRHI/Private/VulkanPendingState.cpp
--- a/Engine/Source/Runtime/VulkanRHI/Private/VulkanPendingState.cpp
+++ b/Engine/Source/Runtime/VulkanRHI/Private/VulkanPendingState.cpp
@@ -329,19 +329,6 @@ FVulkanPendingComputeState::~FVulkanPendingComputeState()
 	}
 }
 
-void FVulkanPendingComputeState::SetSRVForUBResource(uint32 DescriptorSet, uint32 BindingIndex, FVulkanShaderResourceView* SRV)
-{
-	check(SRV);
-	CurrentState->SetSRV(true, DescriptorSet, BindingIndex, SRV);
-}
-
-void FVulkanPendingComputeState::SetUAVForUBResource(uint32 DescriptorSet, uint32 BindingIndex, FVulkanUnorderedAccessView* UAV)
-{
-	check(UAV);
-	CurrentState->SetUAV(true, DescriptorSet, BindingIndex, UAV);
-}
-
-
 void FVulkanPendingComputeState::PrepareForDispatch(FVulkanCommandListContext& Context)
 {
 #if VULKAN_ENABLE_AGGRESSIVE_STATS
@@ -885,17 +872,6 @@ void FVulkanPendingGfxState::UpdateInputAttachments(FVulkanFramebuffer* Framebuf
 	}
 }
 
-void FVulkanPendingGfxState::SetSRVForUBResource(uint8 DescriptorSet, uint32 BindingIndex, FVulkanShaderResourceView* SRV)
-{
-	check(SRV);
-	CurrentState->SetSRV(false, DescriptorSet, BindingIndex, SRV);
-}
-
-void FVulkanPendingGfxState::SetUAVForUBResource(uint8 DescriptorSet, uint32 BindingIndex, FVulkanUnorderedAccessView* UAV)
-{
-	check(UAV);
-	CurrentState->SetUAV(false, DescriptorSet, BindingIndex, UAV);
-}
 
 int32 GDSetCacheTargetSetsPerPool = 4096;
 FAutoConsoleVariableRef CVarDSetCacheTargetSetsPerPool(
diff --git a/Engine/Source/Runtime/VulkanRHI/Private/VulkanPendingState.h b/Engine/Source/Runtime/VulkanRHI/Private/VulkanPendingState.h
--- a/Engine/Source/Runtime/VulkanRHI/Private/VulkanPendingState.h
+++ b/Engine/Source/Runtime/VulkanRHI/Private/VulkanPendingState.h
@@ -61,38 +61,9 @@ public:
 		return CurrentPipeline ? ResourceCast(CurrentPipeline->GetComputeShader()) : nullptr;
 	}
 
-	void SetUAVForUBResource(uint32 DescriptorSet, uint32 BindingIndex, FVulkanUnorderedAccessView* UAV);
-
-	inline void SetUAVForStage(uint32 UAVIndex, FVulkanUnorderedAccessView* UAV)
-	{
-		SetUAVForUBResource(ShaderStage::Compute, UAVIndex, UAV);
-	}
-
-	inline void SetTextureForStage(uint32 TextureIndex, const FVulkanTexture* Texture, VkImageLayout Layout)
-	{
-		CurrentState->SetTexture(ShaderStage::Compute, TextureIndex, Texture, Layout);
-	}
-
-	inline void SetSamplerStateForStage(uint32 SamplerIndex, FVulkanSamplerState* Sampler)
-	{
-		CurrentState->SetSamplerState(ShaderStage::Compute, SamplerIndex, Sampler);
-	}
-
-	inline void SetTextureForUBResource(int32 DescriptorSet, uint32 BindingIndex, const FVulkanTexture* Texture, VkImageLayout Layout)
-	{
-		CurrentState->SetTexture(DescriptorSet, BindingIndex, Texture, Layout);
-	}
-
-	void SetSRVForUBResource(uint32 DescriptorSet, uint32 BindingIndex, FVulkanShaderResourceView* SRV);
-
-	inline void SetSRVForStage(uint32 SRVIndex, FVulkanShaderResourceView* SRV)
-	{
-		SetSRVForUBResource(ShaderStage::Compute, SRVIndex, SRV);
-	}
-
-	inline void SetPackedGlobalShaderParameter(uint32 BufferIndex, uint32 Offset, uint32 NumBytes, const void* NewValue)
+	FVulkanPackedUniformBufferState& GetUniformBuffers(uint8 Stage)
 	{
-		CurrentState->SetPackedGlobalShaderParameter(BufferIndex, Offset, NumBytes, NewValue);
+		return CurrentState->GetUniformBuffers(Stage);
 	}
 
 	inline void SetSamplerStateForUBResource(uint32 DescriptorSet, uint32 BindingIndex, FVulkanSamplerState* Sampler)
@@ -105,6 +76,11 @@ public:
 		PipelineStates.Remove(Pipeline);
 	}
 
+	FVulkanComputePipelineDescriptorState& GetCurrentState()
+	{
+		return *CurrentState;
+	}
+
 protected:
 	FVulkanComputePipeline* CurrentPipeline = nullptr;
 	FVulkanComputePipelineDescriptorState* CurrentState = nullptr;
@@ -237,49 +213,15 @@ public:
 		CurrentPipeline->Bind(CmdBuffer);
 	}
 
-	void SetTextureForStage(ShaderStage::EStage Stage, uint32 ParameterIndex, const FVulkanTexture* Texture, VkImageLayout Layout)
-	{
-		CurrentState->SetTexture(Stage, ParameterIndex, Texture, Layout);
-	}
-
-	void SetTextureForUBResource(uint8 DescriptorSet, uint32 BindingIndex, const FVulkanTexture* Texture, VkImageLayout Layout)
-	{
-		CurrentState->SetTexture(DescriptorSet, BindingIndex, Texture, Layout);
-	}
-
 	template<bool bDynamic>
 	void SetUniformBuffer(uint8 DescriptorSet, uint32 BindingIndex, const FVulkanUniformBuffer* UniformBuffer)
 	{
 		CurrentState->SetUniformBuffer<bDynamic>(DescriptorSet, BindingIndex, UniformBuffer);
 	}
 
-	void SetUAVForUBResource(uint8 DescriptorSet, uint32 BindingIndex, FVulkanUnorderedAccessView* UAV);
-
-	void SetUAVForStage(ShaderStage::EStage Stage, uint32 ParameterIndex, FVulkanUnorderedAccessView* UAV)
-	{
-		SetUAVForUBResource(Stage, ParameterIndex, UAV);
-	}
-
-	void SetSRVForUBResource(uint8 DescriptorSet, uint32 BindingIndex, FVulkanShaderResourceView* SRV);
-
-	void SetSRVForStage(ShaderStage::EStage Stage, uint32 ParameterIndex, FVulkanShaderResourceView* SRV)
-	{
-		SetSRVForUBResource(Stage, ParameterIndex, SRV);
-	}
-
-	void SetSamplerStateForStage(ShaderStage::EStage Stage, uint32 ParameterIndex, FVulkanSamplerState* Sampler)
-	{
-		CurrentState->SetSamplerState(Stage, ParameterIndex, Sampler);
-	}
-
-	void SetSamplerStateForUBResource(uint32 DescriptorSet, uint32 BindingIndex, FVulkanSamplerState* Sampler)
+	FVulkanPackedUniformBufferState& GetUniformBuffers(uint8 Stage)
 	{
-		CurrentState->SetSamplerState(DescriptorSet, BindingIndex, Sampler);
-	}
-
-	void SetPackedGlobalShaderParameter(ShaderStage::EStage Stage, uint32 BufferIndex, uint32 Offset, uint32 NumBytes, const void* NewValue)
-	{
-		CurrentState->SetPackedGlobalShaderParameter(Stage, BufferIndex, Offset, NumBytes, NewValue);
+		return CurrentState->GetUniformBuffers(Stage);
 	}
 
 	void PrepareForDraw(FVulkanCommandListContext& Context);
@@ -301,6 +243,11 @@ public:
 		PipelineStates.Remove(Pipeline);
 	}
 
+	FVulkanGraphicsPipelineDescriptorState& GetCurrentState()
+	{
+		return *CurrentState;
+	}
+
 protected:
 	TArray<VkViewport, TInlineAllocator<2>> Viewports;
 	TArray<VkRect2D, TInlineAllocator<2>> Scissors;
diff --git a/Engine/Source/Runtime/VulkanRHI/Private/VulkanPipelineState.cpp b/Engine/Source/Runtime/VulkanRHI/Private/VulkanPipelineState.cpp
--- a/Engine/Source/Runtime/VulkanRHI/Private/VulkanPipelineState.cpp
+++ b/Engine/Source/Runtime/VulkanRHI/Private/VulkanPipelineState.cpp
@@ -41,14 +41,12 @@ static bool ShouldAlwaysWriteDescriptors()
 
 FVulkanComputePipelineDescriptorState::FVulkanComputePipelineDescriptorState(FVulkanDevice& InDevice, FVulkanComputePipeline* InComputePipeline)
 	: FVulkanCommonPipelineDescriptorState(InDevice, ShaderStage::NumComputeStages, InComputePipeline->UsesBindless())
-	, PackedUniformBuffersMask(0)
-	, PackedUniformBuffersDirty(0)
 	, ComputePipeline(InComputePipeline)
 {
 	LLM_SCOPE_VULKAN(ELLMTagVulkan::VulkanShaders);
 	check(InComputePipeline);
 	const FVulkanShaderHeader& CodeHeader = InComputePipeline->GetShaderCodeHeader();
-	PackedUniformBuffers.Init(CodeHeader, PackedUniformBuffersMask);
+	PackedUniformBufferState.Init(CodeHeader);
 
 	DescriptorSetsLayout = &InComputePipeline->GetLayout().GetDescriptorSetsLayout();
 
@@ -192,13 +190,13 @@ bool FVulkanComputePipelineDescriptorState::InternalUpdateDescriptorSets(FVulkan
 		return false;
 	}
 
-	if (PackedUniformBuffersDirty != 0)
+	if (PackedUniformBufferState.Dirty != 0)
 	{
 #if VULKAN_ENABLE_AGGRESSIVE_STATS
 		SCOPE_CYCLE_COUNTER(STAT_VulkanApplyPackedUniformBuffers);
 #endif
-		SubmitPackedUniformBuffers<bUseDynamicGlobalUBs>(DSWriter[ShaderStage::Compute], UpdatePackedUniformBuffers(PackedUniformBuffers, Context));
-		PackedUniformBuffersDirty = 0;
+		SubmitPackedUniformBuffers<bUseDynamicGlobalUBs>(DSWriter[ShaderStage::Compute], UpdatePackedUniformBuffers(PackedUniformBufferState.UniformBuffers, Context));
+		PackedUniformBufferState.Dirty = 0;
 	}
 
 	// We are not using UseVulkanDescriptorCache() for compute pipelines
@@ -255,10 +253,10 @@ void FVulkanComputePipelineDescriptorState::UpdateBindlessDescriptors(FVulkanCom
 	VkDescriptorBufferInfo ConstantBufferInfo;
 
 	// UBs are currently set from a fresh batch of descriptors for every call, so ignore PackedUniformBuffersDirty
-	check(PackedUniformBuffersMask <= 1);
-	if (PackedUniformBuffersMask != 0)
+	check(PackedUniformBufferState.Mask <= 1);
+	if (PackedUniformBufferState.Mask != 0)
 	{
-		const FPackedUniformBuffers::FPackedBuffer& StagedUniformBuffer = PackedUniformBuffers.GetBuffer();
+		const FPackedUniformBuffers::FPackedBuffer& StagedUniformBuffer = PackedUniformBufferState.UniformBuffers.GetBuffer();
 		const int32 UBSize = StagedUniformBuffer.Num();
 		const VkDeviceSize UBOffsetAlignment = Device.GetLimits().minUniformBufferOffsetAlignment;
 
@@ -278,7 +276,7 @@ void FVulkanComputePipelineDescriptorState::UpdateBindlessDescriptors(FVulkanCom
 		ConstantBufferInfo.offset = TempAllocation.Offset;
 		ConstantBufferInfo.range = UBSize;
 
-		PackedUniformBuffersDirty = 0;
+		PackedUniformBufferState.Dirty = 0;
 	}
 
 	// Send to descriptor manager
@@ -290,8 +288,6 @@ FVulkanGraphicsPipelineDescriptorState::FVulkanGraphicsPipelineDescriptorState(F
 	, GfxPipeline(InGfxPipeline)
 {
 	LLM_SCOPE_VULKAN(ELLMTagVulkan::VulkanShaders);
-	FMemory::Memzero(PackedUniformBuffersMask);
-	FMemory::Memzero(PackedUniformBuffersDirty);
 
 	check(InGfxPipeline && InGfxPipeline->Layout && InGfxPipeline->Layout->IsGfxLayout());
 	DescriptorSetsLayout = &InGfxPipeline->Layout->GetDescriptorSetsLayout();
@@ -302,62 +298,46 @@ FVulkanGraphicsPipelineDescriptorState::FVulkanGraphicsPipelineDescriptorState(F
 
 	TStaticArray<const FVulkanShaderHeader*, ShaderStage::NumGraphicsStages> StageHeaders(InPlace, nullptr);
 
-	uint64 VertexShaderKey = InGfxPipeline->GetShaderKey(SF_Vertex);
-	if (VertexShaderKey)
+	auto InitStage = [this, &StageHeaders](const FVulkanShader* InShader, ShaderStage::EStage InStage)
 	{
-		const FVulkanVertexShader* VertexShader = ShaderFactory.LookupShader<FVulkanVertexShader>(InGfxPipeline->GetShaderKey(SF_Vertex));
-		check(VertexShader);
+		check(InShader);
 
-		PackedUniformBuffers[ShaderStage::Vertex].Init(VertexShader->GetCodeHeader(), PackedUniformBuffersMask[ShaderStage::Vertex]);
-		UsedSetsMask |= VertexShader->GetCodeHeader().Bindings.Num() ? (1u << ShaderStage::Vertex) : 0u;
-		StageHeaders[ShaderStage::Vertex] = &VertexShader->GetCodeHeader();
+		PackedUniformBufferStates[InStage].Init(InShader->GetCodeHeader());
+		UsedSetsMask |= InShader->GetCodeHeader().Bindings.Num() ? (1u << InStage) : 0u;
+		StageHeaders[InStage] = &InShader->GetCodeHeader();
+	};
+
+	if (uint64 VertexShaderKey = InGfxPipeline->GetShaderKey(SF_Vertex))
+	{
+		const FVulkanVertexShader* VertexShader = ShaderFactory.LookupShader<FVulkanVertexShader>(VertexShaderKey);
+		InitStage(VertexShader, ShaderStage::Vertex);
 	}
 
-	uint64 PixelShaderKey = InGfxPipeline->GetShaderKey(SF_Pixel);
-	if (PixelShaderKey)
+	if (uint64 PixelShaderKey = InGfxPipeline->GetShaderKey(SF_Pixel))
 	{
 		const FVulkanPixelShader* PixelShader = ShaderFactory.LookupShader<FVulkanPixelShader>(PixelShaderKey);
-		check(PixelShader);
-
-		PackedUniformBuffers[ShaderStage::Pixel].Init(PixelShader->GetCodeHeader(), PackedUniformBuffersMask[ShaderStage::Pixel]);
-		UsedSetsMask |= PixelShader->GetCodeHeader().Bindings.Num() ? (1u << ShaderStage::Pixel) : 0u;
-		StageHeaders[ShaderStage::Pixel] = &PixelShader->GetCodeHeader();
+		InitStage(PixelShader, ShaderStage::Pixel);
 	}
 
 #if PLATFORM_SUPPORTS_MESH_SHADERS
-	uint64 MeshShaderKey = InGfxPipeline->GetShaderKey(SF_Mesh);
-	if (MeshShaderKey)
+	if (uint64 MeshShaderKey = InGfxPipeline->GetShaderKey(SF_Mesh))
 	{
 		const FVulkanMeshShader* MeshShader = ShaderFactory.LookupShader<FVulkanMeshShader>(MeshShaderKey);
-		check(MeshShader);
-
-		PackedUniformBuffers[ShaderStage::Mesh].Init(MeshShader->GetCodeHeader(), PackedUniformBuffersMask[ShaderStage::Mesh]);
-		UsedSetsMask |= MeshShader->GetCodeHeader().Bindings.Num() ? (1u << ShaderStage::Mesh) : 0u;
-		StageHeaders[ShaderStage::Mesh] = &MeshShader->GetCodeHeader();
+		InitStage(MeshShader, ShaderStage::Mesh);
 	}
 
-	uint64 TaskShaderKey = InGfxPipeline->GetShaderKey(SF_Amplification);
-	if (TaskShaderKey)
+	if (uint64 TaskShaderKey = InGfxPipeline->GetShaderKey(SF_Amplification))
 	{
 		const FVulkanTaskShader* TaskShader = ShaderFactory.LookupShader<FVulkanTaskShader>(TaskShaderKey);
-		check(TaskShader);
-
-		PackedUniformBuffers[ShaderStage::Task].Init(TaskShader->GetCodeHeader(), PackedUniformBuffersMask[ShaderStage::Task]);
-		UsedSetsMask |= TaskShader->GetCodeHeader().Bindings.Num() ? (1u << ShaderStage::Task) : 0u;
-		StageHeaders[ShaderStage::Task] = &TaskShader->GetCodeHeader();
+		InitStage(TaskShader, ShaderStage::Task);
 	}
 #endif
 
 #if VULKAN_SUPPORTS_GEOMETRY_SHADERS
-	uint64 GeometryShaderKey = InGfxPipeline->GetShaderKey(SF_Geometry);
-	if (GeometryShaderKey)
+	if (uint64 GeometryShaderKey = InGfxPipeline->GetShaderKey(SF_Geometry))
 	{
 		const FVulkanGeometryShader* GeometryShader = ShaderFactory.LookupShader<FVulkanGeometryShader>(GeometryShaderKey);
-		check(GeometryShader);
-
-		PackedUniformBuffers[ShaderStage::Geometry].Init(GeometryShader->GetCodeHeader(), PackedUniformBuffersMask[ShaderStage::Geometry]);
-		UsedSetsMask |= GeometryShader->GetCodeHeader().Bindings.Num() ? (1u << ShaderStage::Geometry) : 0u;
-		StageHeaders[ShaderStage::Geometry] = &GeometryShader->GetCodeHeader();
+		InitStage(GeometryShader, ShaderStage::Geometry);
 	}
 #endif
 
@@ -368,8 +348,7 @@ FVulkanGraphicsPipelineDescriptorState::FVulkanGraphicsPipelineDescriptorState(F
 
 	for (int32 StageIndex = 0; StageIndex < StageHeaders.Num(); ++StageIndex)
 	{
-		const FVulkanShaderHeader* CodeHeader = StageHeaders[StageIndex];
-		if (CodeHeader)
+		if (const FVulkanShaderHeader* CodeHeader = StageHeaders[StageIndex])
 		{
 			for (const FVulkanShaderHeader::FGlobalSamplerInfo& GlobalSamplerInfo : CodeHeader->GlobalSamplerInfos)
 			{
@@ -400,10 +379,11 @@ bool FVulkanGraphicsPipelineDescriptorState::InternalUpdateDescriptorSets(FVulka
 #endif
 		for (int32 Stage = 0; Stage < ShaderStage::NumGraphicsStages; ++Stage)
 		{
-			if (PackedUniformBuffersDirty[Stage] != 0)
+			FVulkanPackedUniformBufferState& State = PackedUniformBufferStates[Stage];
+			if (State.Dirty != 0)
 			{
-				MarkDirty(SubmitPackedUniformBuffers<bUseDynamicGlobalUBs>(DSWriter[Stage], UpdatePackedUniformBuffers(PackedUniformBuffers[Stage], Context)));
-				PackedUniformBuffersDirty[Stage] = 0;
+				MarkDirty(SubmitPackedUniformBuffers<bUseDynamicGlobalUBs>(DSWriter[Stage], UpdatePackedUniformBuffers(State.UniformBuffers, Context)));
+				State.Dirty = 0;
 			}
 		}
 	}
@@ -518,10 +498,12 @@ void FVulkanGraphicsPipelineDescriptorState::UpdateBindlessDescriptors(FVulkanCo
 			}
 
 			// UBs are currently set from a fresh batch of descriptors for every call, so ignore PackedUniformBuffersDirty
-			check(PackedUniformBuffersMask[Stage] <= 1);
-			if ((UBIndex == 0) && (PackedUniformBuffersMask[Stage] != 0))
+			FVulkanPackedUniformBufferState& State = PackedUniformBufferStates[Stage];
+
+			check(State.Mask <= 1);
+			if (UBIndex == 0 && State.Mask != 0)
 			{
-				const FPackedUniformBuffers::FPackedBuffer& StagedUniformBuffer = PackedUniformBuffers[Stage].GetBuffer();
+				const FPackedUniformBuffers::FPackedBuffer& StagedUniformBuffer = State.UniformBuffers.GetBuffer();
 				const int32 UBSize = StagedUniformBuffer.Num();
 
 				VulkanRHI::FVulkanAllocation TempAllocation;
@@ -536,7 +518,7 @@ void FVulkanGraphicsPipelineDescriptorState::UpdateBindlessDescriptors(FVulkanCo
 				ConstantBufferInfo.offset = TempAllocation.Offset;
 				ConstantBufferInfo.range = UBSize;
 
-				PackedUniformBuffersDirty[Stage] = 0;
+				State.Dirty = 0;
 			}
 		}
 	}
diff --git a/Engine/Source/Runtime/VulkanRHI/Private/VulkanPipelineState.h b/Engine/Source/Runtime/VulkanRHI/Private/VulkanPipelineState.h
--- a/Engine/Source/Runtime/VulkanRHI/Private/VulkanPipelineState.h
+++ b/Engine/Source/Runtime/VulkanRHI/Private/VulkanPipelineState.h
@@ -177,6 +177,27 @@ protected:
 	const bool bUseBindless;
 };
 
+struct FVulkanPackedUniformBufferState
+{
+	FPackedUniformBuffers UniformBuffers;
+	uint32 Mask = 0;
+	uint32 Dirty = 0;
+
+	void Init(const FVulkanShaderHeader& CodeHeader)
+	{
+		UniformBuffers.Init(CodeHeader, Mask);
+	}
+
+	void Reset()
+	{
+		Dirty = Mask;
+	}
+
+	void SetParameter(uint32 ByteOffset, uint32 NumBytes, const void* NewValue)
+	{
+		UniformBuffers.SetPackedGlobalParameter(ByteOffset, NumBytes, NewValue, Dirty);
+	}
+};
 
 class FVulkanComputePipelineDescriptorState : public FVulkanCommonPipelineDescriptorState
 {
@@ -190,41 +211,39 @@ public:
 	void Reset()
 	{
 		FVulkanCommonPipelineDescriptorState::Reset();
-		PackedUniformBuffersDirty = PackedUniformBuffersMask;
+		PackedUniformBufferState.Reset();
+	}
+
+	FVulkanPackedUniformBufferState& GetUniformBuffers(uint8 Stage)
+	{
+		return PackedUniformBufferState;
 	}
 
-	inline void SetPackedGlobalShaderParameter(uint32 BufferIndex, uint32 ByteOffset, uint32 NumBytes, const void* NewValue)
+	void SetParameter(uint32 ByteOffset, uint32 NumBytes, const void* NewValue)
 	{
-		check(BufferIndex == 0);
-		PackedUniformBuffers.SetPackedGlobalParameter(ByteOffset, NumBytes, NewValue, PackedUniformBuffersDirty);
+		PackedUniformBufferState.SetParameter(ByteOffset, NumBytes, NewValue);
 	}
 
 	bool UpdateDescriptorSets(FVulkanCommandListContext& Context)
 	{
 		check(!bUseBindless);
-
-		const bool bUseDynamicGlobalUBs = (GDynamicGlobalUBs->GetInt() > 0);
-		if (bUseDynamicGlobalUBs)
+		if (GDynamicGlobalUBs->GetInt() > 0)
 		{
 			return InternalUpdateDescriptorSets<true>(Context);
 		}
-		else
-		{
-			return InternalUpdateDescriptorSets<false>(Context);
-		}
+
+		return InternalUpdateDescriptorSets<false>(Context);
 	}
 
 	void UpdateBindlessDescriptors(FVulkanCommandListContext& Context);
 
-	inline void BindDescriptorSets(VkCommandBuffer CmdBuffer)
+	void BindDescriptorSets(VkCommandBuffer CmdBuffer)
 	{
 		Bind(CmdBuffer, ComputePipeline->GetLayout().GetPipelineLayout(), VK_PIPELINE_BIND_POINT_COMPUTE);
 	}
 
 protected:
-	FPackedUniformBuffers PackedUniformBuffers;
-	uint32 PackedUniformBuffersMask;
-	uint32 PackedUniformBuffersDirty;
+	FVulkanPackedUniformBufferState PackedUniformBufferState;
 
 	FVulkanComputePipeline* ComputePipeline;
 
@@ -244,10 +263,14 @@ public:
 		GfxPipeline->Release();
 	}
 
-	inline void SetPackedGlobalShaderParameter(uint8 Stage, uint32 BufferIndex, uint32 ByteOffset, uint32 NumBytes, const void* NewValue)
+	FVulkanPackedUniformBufferState& GetUniformBuffers(uint8 Stage)
+	{
+		return PackedUniformBufferStates[Stage];
+	}
+
+	void SetParameter(uint8 Stage, uint32 ByteOffset, uint32 NumBytes, const void* NewValue)
 	{
-		check(BufferIndex == 0);
-		PackedUniformBuffers[Stage].SetPackedGlobalParameter(ByteOffset, NumBytes, NewValue, PackedUniformBuffersDirty[Stage]);
+		PackedUniformBufferStates[Stage].SetParameter(ByteOffset, NumBytes, NewValue);
 	}
 
 	bool UpdateDescriptorSets(FVulkanCommandListContext& Context)
@@ -274,15 +297,17 @@ public:
 
 	void Reset()
 	{
-		FMemory::Memcpy(PackedUniformBuffersDirty, PackedUniformBuffersMask);
+		for (FVulkanPackedUniformBufferState& State : PackedUniformBufferStates)
+		{
+			State.Reset();
+		}
+
 		FVulkanCommonPipelineDescriptorState::Reset();
 		bIsResourcesDirty = true;
 	}
 
 protected:
-	TStaticArray<FPackedUniformBuffers, ShaderStage::NumGraphicsStages> PackedUniformBuffers;
-	TStaticArray<uint32, ShaderStage::NumGraphicsStages> PackedUniformBuffersMask;
-	TStaticArray<uint32, ShaderStage::NumGraphicsStages> PackedUniformBuffersDirty;
+	TStaticArray<FVulkanPackedUniformBufferState, ShaderStage::NumGraphicsStages> PackedUniformBufferStates;
 
 	FVulkanGraphicsPipelineState* GfxPipeline;
 
