Vulkan binding changes
- Remove uses of the older SetShaderXX context methods in favor of using the Binder approach that other RHIs use now.
- Rename GlobalUniformBuffers to StaticUniformBuffers to match API names.

#Synced-CL 48355415

diff --git a/Engine/Source/Runtime/RHICore/Public/RHICoreShader.h b/Engine/Source/Runtime/RHICore/Public/RHICoreShader.h
--- a/Engine/Source/Runtime/RHICore/Public/RHICoreShader.h
+++ b/Engine/Source/Runtime/RHICore/Public/RHICoreShader.h
@@ -63,7 +63,7 @@ inline void InitStaticUniformBufferSlots(FRHIShaderData* ShaderData)
 
 template <typename TApplyFunction>
 void ApplyStaticUniformBuffers(
-	FRHIShader* Shader,
+	const FRHIShader* Shader,
 	const TArray<FRHIUniformBuffer*>& UniformBuffers,
 	TApplyFunction&& ApplyFunction
 )
@@ -94,7 +94,7 @@ void ApplyStaticUniformBuffers(
 template <typename TRHIContext, typename TRHIShader>
 void ApplyStaticUniformBuffers(
 	TRHIContext* CommandContext,
-	TRHIShader* Shader,
+	const TRHIShader* Shader,
 	const TArray<FRHIUniformBuffer*>& UniformBuffers)
 {
 	ApplyStaticUniformBuffers(Shader, UniformBuffers,
diff --git a/Engine/Source/Runtime/VulkanRHI/Private/VulkanCommands.cpp b/Engine/Source/Runtime/VulkanRHI/Private/VulkanCommands.cpp
--- a/Engine/Source/Runtime/VulkanRHI/Private/VulkanCommands.cpp
+++ b/Engine/Source/Runtime/VulkanRHI/Private/VulkanCommands.cpp
@@ -38,6 +38,33 @@ static_assert(STRUCT_OFFSET(FRHIDispatchIndirectParameters, ThreadGroupCountX) =
 static_assert(STRUCT_OFFSET(FRHIDispatchIndirectParameters, ThreadGroupCountY) == STRUCT_OFFSET(VkDispatchIndirectCommand, y), "FRHIDispatchIndirectParameters Y dimension is wrong.");
 static_assert(STRUCT_OFFSET(FRHIDispatchIndirectParameters, ThreadGroupCountZ) == STRUCT_OFFSET(VkDispatchIndirectCommand, z), "FRHIDispatchIndirectParameters Z dimension is wrong.");
 
+inline const FRHIShader* RHIShaderFromVulkanShader(const FVulkanShader* InShader, EShaderFrequency InFrequency)
+{
+	switch (InFrequency)
+	{
+	case SF_Vertex:
+		return static_cast<const FVulkanVertexShader*>(InShader);
+#if PLATFORM_SUPPORTS_MESH_SHADERS
+	case SF_Mesh:
+		return static_cast<const FVulkanMeshShader*>(InShader);
+	case SF_Amplification:
+		return static_cast<const FVulkanTaskShader*>(InShader);
+#endif
+	case SF_Pixel:
+		return static_cast<const FVulkanPixelShader*>(InShader);
+#if VULKAN_SUPPORTS_GEOMETRY_SHADERS
+	case SF_Geometry:
+		return static_cast<const FVulkanGeometryShader*>(InShader);
+#endif
+	case SF_Compute:
+		return static_cast<const FVulkanComputeShader*>(InShader);
+
+	default:
+		checkf(0, TEXT("Undefined FRHIShader Frequency %d!"), (int32)InFrequency);
+		return nullptr;
+	}
+}
+
 static FORCEINLINE ShaderStage::EStage GetAndVerifyShaderStage(FRHIGraphicsShader* ShaderRHI, FVulkanPendingGfxState* PendingGfxState)
 {
 	switch (ShaderRHI->GetFrequency())
@@ -124,23 +151,43 @@ void FVulkanCommandListContext::RHISetStreamSource(uint32 StreamIndex, FRHIBuffe
 struct FVulkanResourceBinder
 {
 	FVulkanCommandListContext& Context;
+	const FVulkanShader* Shader;
 	const EShaderFrequency Frequency;
 	const ShaderStage::EStage Stage;
 	FVulkanCommonPipelineDescriptorState& State;
+	FVulkanPackedUniformBufferState& UniformBuffer;
 
 	// Even if no resource are bound when bindless is enabled, we still need to process the Resource Table to go through proper validation
 	const bool bBindless;
 
 	template <class PendingStateType>
-	FVulkanResourceBinder(FVulkanCommandListContext& InContext, EShaderFrequency InFrequency, PendingStateType* InPendingState, bool InUsesBindless)
+	FVulkanResourceBinder(FVulkanCommandListContext& InContext, PendingStateType* InPendingState, EShaderFrequency InFrequency, const FVulkanShader* InShader)
 		: Context(InContext)
+		, Shader(InShader)
 		, Frequency(InFrequency)
-		, Stage((InFrequency == SF_Compute) ? ShaderStage::Compute : ShaderStage::GetStageForFrequency(InFrequency))
+		, Stage(ShaderStage::GetStageForFrequency(Frequency))
 		, State(InPendingState->GetCurrentState())
-		, bBindless(InUsesBindless)
+		, UniformBuffer(InPendingState->GetUniformBuffers(Stage))
+		, bBindless(InShader->UsesBindless())
 	{
 	}
 
+	void SetParameter(uint32 BaseIndex, uint32 NumBytes, const void* NewValue)
+	{
+		UniformBuffer.SetParameter(BaseIndex, NumBytes, NewValue);
+	}
+
+#if PLATFORM_SUPPORTS_BINDLESS_RENDERING
+	void SetBindlessHandle(const FRHIDescriptorHandle& Handle, uint32 Offset)
+	{
+		if (Handle.IsValid())
+		{
+			const uint32 BindlessIndex = Handle.GetIndex();
+			SetParameter(Offset, 4, reinterpret_cast<const uint8*>(&BindlessIndex));
+		}
+	}
+#endif
+
 	void SetUAV(FRHIUnorderedAccessView* UAV, uint16 Index, bool bClearResources = false)
 	{
 		if (bClearResources)
@@ -199,7 +246,7 @@ void FVulkanCommandListContext::SetResourcesFromTables(const ShaderType* Shader)
 
 	if (Frequency == SF_Compute)
 	{
-		FVulkanResourceBinder Binder(*this, Frequency, PendingComputeState, Shader->UsesBindless());
+		FVulkanResourceBinder Binder(*this, PendingComputeState, SF_Compute, Shader);
 		UE::RHI::Private::SetUniformBufferResourcesFromTables(
 			Binder
 			, *Shader
@@ -212,7 +259,7 @@ void FVulkanCommandListContext::SetResourcesFromTables(const ShaderType* Shader)
 	}
 	else
 	{
-		FVulkanResourceBinder Binder(*this, Frequency, PendingGfxState, Shader->UsesBindless());
+		FVulkanResourceBinder Binder(*this, PendingGfxState, Frequency, Shader);
 		UE::RHI::Private::SetUniformBufferResourcesFromTables(
 			Binder
 			, *Shader
@@ -225,6 +272,162 @@ void FVulkanCommandListContext::SetResourcesFromTables(const ShaderType* Shader)
 	}
 }
 
+static void BindUniformBuffer(
+	FVulkanCommandListContext& Context,
+	FVulkanCommonPipelineDescriptorState& State,
+	const FVulkanShader* Shader,
+	EShaderFrequency Frequency,
+	ShaderStage::EStage Stage,
+	uint32 BufferIndex,
+	FVulkanUniformBuffer* UniformBuffer
+)
+{
+#if VULKAN_ENABLE_AGGRESSIVE_STATS
+	SCOPE_CYCLE_COUNTER(STAT_VulkanSetUniformBufferTime);
+#endif
+
+	const FVulkanShaderHeader& CodeHeader = Shader->GetCodeHeader();
+	checkfSlow(!CodeHeader.UniformBufferInfos[BufferIndex].LayoutHash || (CodeHeader.UniformBufferInfos[BufferIndex].LayoutHash == UniformBuffer->GetLayout().GetHash()),
+		TEXT("Mismatched UB layout! Got hash 0x%x, expected 0x%x!"), UniformBuffer->GetLayout().GetHash(), CodeHeader.UniformBufferInfos[BufferIndex].LayoutHash);
+
+	if (UniformBuffer->IsUniformView())
+	{
+		UniformBuffer->SetupUniformBufferView();
+	}
+
+	bool bHasResources = false;
+	if (BufferIndex < CodeHeader.NumBoundUniformBuffers)
+	{
+		checkSlow(UniformBuffer->GetLayout().ConstantBufferSize > 0);
+
+		const VkDescriptorType DescriptorType = State.GetDescriptorType(Stage, BufferIndex);
+		const bool bDynamic = (DescriptorType == VK_DESCRIPTOR_TYPE_UNIFORM_BUFFER_DYNAMIC);
+
+		State.SetUniformBuffer(Stage, BufferIndex, UniformBuffer, bDynamic);
+
+		bHasResources = (CodeHeader.UniformBufferInfos[BufferIndex].bHasResources != 0);
+	}
+	else
+	{
+		// If the buffer has no bindings, then it is as resource only ub
+		bHasResources = true;
+	}
+
+	if (bHasResources)
+	{
+		checkSlow(Frequency < SF_NumStandardFrequencies);
+		check(BufferIndex < FVulkanCommandListContext::MAX_UNIFORM_BUFFERS_PER_SHADER_STAGE);
+		Context.BoundUniformBuffers[Frequency][BufferIndex] = UniformBuffer;
+		Context.DirtyUniformBuffers[Frequency] |= (1 << BufferIndex);
+	}
+}
+
+static void SetShaderParametersOnBinder(
+	FVulkanResourceBinder& Binder
+	, TConstArrayView<uint8> InParametersData
+	, TConstArrayView<FRHIShaderParameter> InParameters
+	, TConstArrayView<FRHIShaderParameterResource> InResourceParameters
+	, TConstArrayView<FRHIShaderParameterResource> InBindlessParameters)
+{
+	const EShaderFrequency Frequency = Binder.Frequency;
+
+	for (const FRHIShaderParameter& Parameter : InParameters)
+	{
+		checkSlow(Parameter.BufferIndex == 0);
+		Binder.SetParameter(Parameter.BaseIndex, Parameter.ByteSize, &InParametersData[Parameter.ByteOffset]);
+	}
+
+#if PLATFORM_SUPPORTS_BINDLESS_RENDERING
+	for (const FRHIShaderParameterResource& Parameter : InBindlessParameters)
+	{
+		if (FRHIResource* Resource = Parameter.Resource)
+		{
+			if (Parameter.Type == FRHIShaderParameterResource::EType::UniformBuffer)
+			{
+				// ue-todo: christopher.waters - Enable SetAllUniformBufferResourcesForBindless once it's checked in.
+				//UE::RHI::Private::SetAllUniformBufferResourcesForBindless(Binder, static_cast<FRHIUniformBuffer*>(Resource));
+				checkNoEntry();
+			}
+			else
+			{
+				FRHIDescriptorHandle Handle;
+
+				switch (Parameter.Type)
+				{
+				case FRHIShaderParameterResource::EType::Texture:
+					Handle = static_cast<FRHITexture*>(Resource)->GetDefaultBindlessHandle();
+					Binder.SetTexture(static_cast<FRHITexture*>(Resource), Parameter.Index);
+					break;
+				case FRHIShaderParameterResource::EType::ResourceView:
+					Handle = static_cast<FRHIShaderResourceView*>(Resource)->GetBindlessHandle();
+					Binder.SetSRV(static_cast<FRHIShaderResourceView*>(Resource), Parameter.Index);
+					break;
+				case FRHIShaderParameterResource::EType::UnorderedAccessView:
+					Handle = static_cast<FRHIUnorderedAccessView*>(Resource)->GetBindlessHandle();
+					Binder.SetUAV(static_cast<FRHIUnorderedAccessView*>(Resource), Parameter.Index, true);
+					break;
+				case FRHIShaderParameterResource::EType::Sampler:
+					Handle = static_cast<FRHISamplerState*>(Resource)->GetBindlessHandle();
+					Binder.SetSampler(static_cast<FRHISamplerState*>(Resource), Parameter.Index);
+					break;
+				case FRHIShaderParameterResource::EType::ResourceCollection:
+					Handle = static_cast<FRHIResourceCollection*>(Resource)->GetBindlessHandle();
+					Binder.SetResourceCollection(static_cast<FRHIResourceCollection*>(Resource), Parameter.Index);
+					break;
+				}
+
+				checkf(Handle.IsValid(), TEXT("Vulkan resource did not provide a valid descriptor handle. Please validate that all Vulkan types can provide this or that the resource is still valid."));
+				Binder.SetBindlessHandle(Handle, Parameter.Index);
+			}
+		}
+	}
+#endif
+
+	for (const FRHIShaderParameterResource& Parameter : InResourceParameters)
+	{
+		if (Parameter.Type == FRHIShaderParameterResource::EType::UnorderedAccessView)
+		{
+			if (Binder.Frequency == SF_Pixel || Binder.Frequency == SF_Vertex || Binder.Frequency == SF_Compute)
+			{
+				Binder.SetUAV(static_cast<FRHIUnorderedAccessView*>(Parameter.Resource), Parameter.Index, true);
+			}
+			else
+			{
+				checkf(false, TEXT("TShaderRHI Can't have compute shader to be set. UAVs are not supported tessellation and geometry shaders."));
+			}
+		}
+	}
+
+	for (const FRHIShaderParameterResource& Parameter : InResourceParameters)
+	{
+		switch (Parameter.Type)
+		{
+		case FRHIShaderParameterResource::EType::Texture:
+			Binder.SetTexture(static_cast<FRHITexture*>(Parameter.Resource), Parameter.Index);
+			break;
+		case FRHIShaderParameterResource::EType::ResourceView:
+			Binder.SetSRV(static_cast<FRHIShaderResourceView*>(Parameter.Resource), Parameter.Index);
+			break;
+		case FRHIShaderParameterResource::EType::UnorderedAccessView:
+			break;
+		case FRHIShaderParameterResource::EType::Sampler:
+			Binder.SetSampler(static_cast<FRHISamplerState*>(Parameter.Resource), Parameter.Index);
+			break;
+		case FRHIShaderParameterResource::EType::UniformBuffer:
+			BindUniformBuffer(Binder.Context, Binder.State, Binder.Shader, Binder.Frequency, Binder.Stage, Parameter.Index, ResourceCast(static_cast<FRHIUniformBuffer*>(Parameter.Resource)));
+			break;
+#if PLATFORM_SUPPORTS_BINDLESS_RENDERING
+		case FRHIShaderParameterResource::EType::ResourceCollection:
+			Binder.SetResourceCollection(static_cast<FRHIResourceCollection*>(Parameter.Resource), Parameter.Index);
+			break;
+#endif
+		default:
+			checkf(false, TEXT("Unhandled resource type?"));
+			break;
+		}
+	}
+}
+
 void FVulkanCommandListContext::CommitGraphicsResourceTables()
 {
 	checkfSlow(Queue.GetQueueType() == EVulkanQueueType::Graphics, TEXT("Recording a graphic command on a non-graphic queue."));
@@ -315,114 +518,37 @@ void FVulkanCommandListContext::RHIDispatchIndirectComputeShader(FRHIBuffer* Arg
 	VulkanRHI::DebugHeavyWeightBarrier(CommandBufferHandle, 2);
 }
 
-void FVulkanCommandListContext::RHISetUAVParameter(FRHIPixelShader* PixelShaderRHI, uint32 UAVIndex, FRHIUnorderedAccessView* UAVRHI)
-{
-	if (UAVRHI)
-	{
-		PendingGfxState->GetCurrentState().SetUAV(false, ShaderStage::Pixel, UAVIndex, ResourceCast(UAVRHI));
-	}
-}
-
-void FVulkanCommandListContext::RHISetUAVParameter(FRHIComputeShader* ComputeShaderRHI, uint32 UAVIndex, FRHIUnorderedAccessView* UAVRHI)
-{
-	if (UAVRHI)
-	{
-		check(PendingComputeState->GetCurrentShader() == ResourceCast(ComputeShaderRHI));
-		PendingComputeState->GetCurrentState().SetUAV(true, ShaderStage::Compute, UAVIndex, ResourceCast(UAVRHI));
-	}
-}
-
-void FVulkanCommandListContext::RHISetShaderTexture(FRHIGraphicsShader* ShaderRHI, uint32 TextureIndex, FRHITexture* NewTextureRHI)
-{
-	FVulkanTexture* VulkanTexture = ResourceCast(NewTextureRHI);
-	const VkImageLayout ExpectedLayout = FVulkanPipelineBarrier::GetDefaultLayout(*VulkanTexture, ERHIAccess::SRVGraphics);
-
-	const ShaderStage::EStage Stage = GetAndVerifyShaderStage(ShaderRHI, PendingGfxState);
-	PendingGfxState->GetCurrentState().SetTexture(Stage, TextureIndex, VulkanTexture, ExpectedLayout);
-
-	NewTextureRHI->SetLastRenderTime((float)FPlatformTime::Seconds());
-}
-
-void FVulkanCommandListContext::RHISetShaderTexture(FRHIComputeShader* ComputeShaderRHI, uint32 TextureIndex, FRHITexture* NewTextureRHI)
-{
-	check(PendingComputeState->GetCurrentShader() == ResourceCast(ComputeShaderRHI));
-
-	FVulkanTexture* VulkanTexture = ResourceCast(NewTextureRHI);
-	const VkImageLayout ExpectedLayout = FVulkanPipelineBarrier::GetDefaultLayout(*VulkanTexture, ERHIAccess::SRVCompute);
-	PendingComputeState->GetCurrentState().SetTexture(ShaderStage::Compute, TextureIndex, VulkanTexture, ExpectedLayout);
-
-	NewTextureRHI->SetLastRenderTime((float)FPlatformTime::Seconds());
-}
-
-void FVulkanCommandListContext::RHISetShaderResourceViewParameter(FRHIGraphicsShader* ShaderRHI, uint32 TextureIndex, FRHIShaderResourceView* SRVRHI)
-{
-	if (SRVRHI)
-	{
-		const ShaderStage::EStage Stage = GetAndVerifyShaderStage(ShaderRHI, PendingGfxState);
-		PendingGfxState->GetCurrentState().SetSRV(false, Stage, TextureIndex, ResourceCast(SRVRHI));
-	}
-}
-
-void FVulkanCommandListContext::RHISetShaderResourceViewParameter(FRHIComputeShader* ComputeShaderRHI,uint32 TextureIndex, FRHIShaderResourceView* SRVRHI)
+void FVulkanCommandListContext::RHISetShaderParameters(FRHIGraphicsShader* ShaderRHI, TConstArrayView<uint8> InParametersData, TConstArrayView<FRHIShaderParameter> InParameters, TConstArrayView<FRHIShaderParameterResource> InResourceParameters, TConstArrayView<FRHIShaderParameterResource> InBindlessParameters)
 {
-	if (SRVRHI)
-	{
-		check(PendingComputeState->GetCurrentShader() == ResourceCast(ComputeShaderRHI));
-		PendingComputeState->GetCurrentState().SetSRV(true, ShaderStage::Compute, TextureIndex, ResourceCast(SRVRHI));
-	}
-}
-
-void FVulkanCommandListContext::RHISetShaderSampler(FRHIGraphicsShader* ShaderRHI, uint32 SamplerIndex, FRHISamplerState* NewStateRHI)
-{
-	const ShaderStage::EStage Stage = GetAndVerifyShaderStage(ShaderRHI, PendingGfxState);
-	PendingGfxState->GetCurrentState().SetSamplerState(Stage, SamplerIndex, ResourceCast(NewStateRHI));
-}
-
-void FVulkanCommandListContext::RHISetShaderSampler(FRHIComputeShader* ComputeShaderRHI, uint32 SamplerIndex, FRHISamplerState* NewStateRHI)
-{
-	check(PendingComputeState->GetCurrentShader() == ResourceCast(ComputeShaderRHI));
-	PendingComputeState->GetCurrentState().SetSamplerState(ShaderStage::Compute, SamplerIndex, ResourceCast(NewStateRHI));
-}
-
-void FVulkanCommandListContext::RHISetShaderParameter(FRHIGraphicsShader* ShaderRHI, uint32 BufferIndex, uint32 BaseIndex, uint32 NumBytes, const void* NewValue)
-{
-	const ShaderStage::EStage Stage = GetAndVerifyShaderStage(ShaderRHI, PendingGfxState);
-	PendingGfxState->GetCurrentState().GetUniformBuffers(Stage).SetParameter(BaseIndex, NumBytes, NewValue);
-}
+	FVulkanShader* Shader = nullptr;
+	const ShaderStage::EStage Stage = GetAndVerifyShaderStageAndVulkanShader(ShaderRHI, PendingGfxState, Shader);
 
-void FVulkanCommandListContext::RHISetShaderParameter(FRHIComputeShader* ComputeShaderRHI,uint32 BufferIndex, uint32 BaseIndex, uint32 NumBytes, const void* NewValue)
-{
-	check(PendingComputeState->GetCurrentShader() == ResourceCast(ComputeShaderRHI));
-	PendingComputeState->GetCurrentState().GetUniformBuffers(ShaderStage::Compute).SetParameter(BaseIndex, NumBytes, NewValue);
-}
+	FVulkanResourceBinder Binder(*this, PendingGfxState, Shader->Frequency, Shader);
 
-void FVulkanCommandListContext::RHISetShaderParameters(FRHIGraphicsShader* Shader, TConstArrayView<uint8> InParametersData, TConstArrayView<FRHIShaderParameter> InParameters, TConstArrayView<FRHIShaderParameterResource> InResourceParameters, TConstArrayView<FRHIShaderParameterResource> InBindlessParameters)
-{
-	UE::RHICore::RHISetShaderParametersShared(
-		*this
-		, Shader
+	SetShaderParametersOnBinder(
+		Binder
 		, InParametersData
 		, InParameters
 		, InResourceParameters
-		, InBindlessParameters
-	);
+		, InBindlessParameters);
 }
 
-void FVulkanCommandListContext::RHISetShaderParameters(FRHIComputeShader* Shader, TConstArrayView<uint8> InParametersData, TConstArrayView<FRHIShaderParameter> InParameters, TConstArrayView<FRHIShaderParameterResource> InResourceParameters, TConstArrayView<FRHIShaderParameterResource> InBindlessParameters)
+void FVulkanCommandListContext::RHISetShaderParameters(FRHIComputeShader* ShaderRHI, TConstArrayView<uint8> InParametersData, TConstArrayView<FRHIShaderParameter> InParameters, TConstArrayView<FRHIShaderParameterResource> InResourceParameters, TConstArrayView<FRHIShaderParameterResource> InBindlessParameters)
 {
-	UE::RHICore::RHISetShaderParametersShared(
-		*this
-		, Shader
+	check(PendingComputeState->GetCurrentShader() == ResourceCast(ShaderRHI));
+	FVulkanResourceBinder Binder(*this, PendingComputeState, SF_Compute, ResourceCast(ShaderRHI));
+
+	SetShaderParametersOnBinder(
+		Binder
 		, InParametersData
 		, InParameters
 		, InResourceParameters
-		, InBindlessParameters
-	);
+		, InBindlessParameters);
 }
 
 void FVulkanCommandListContext::RHISetStaticUniformBuffers(const FUniformBufferStaticBindings& InUniformBuffers)
 {
-	FMemory::Memzero(GlobalUniformBuffers.GetData(), GlobalUniformBuffers.Num() * sizeof(FRHIUniformBuffer*));
+	FMemory::Memzero(StaticUniformBuffers.GetData(), StaticUniformBuffers.Num() * sizeof(FRHIUniformBuffer*));
 
 	if (const FRHIShaderBindingLayout* Layout = InUniformBuffers.GetShaderBindingLayout())
 	{
@@ -430,8 +556,8 @@ void FVulkanCommandListContext::RHISetStaticUniformBuffers(const FUniformBufferS
 
 		for (int32 Index = 0; Index < InUniformBuffers.GetUniformBufferCount(); ++Index)
 		{
-			GlobalUniformBuffers[Index] = InUniformBuffers.GetUniformBuffer(Index);
-			checkf(GlobalUniformBuffers[Index], TEXT("Static uniform buffer at index %d is referenced in the shader binding layout but is not provided"), Index);
+			StaticUniformBuffers[Index] = InUniformBuffers.GetUniformBuffer(Index);
+			checkf(StaticUniformBuffers[Index], TEXT("Static uniform buffer at index %d is referenced in the shader binding layout but is not provided"), Index);
 		}
 
 		ShaderBindingLayout = Layout;
@@ -440,7 +566,7 @@ void FVulkanCommandListContext::RHISetStaticUniformBuffers(const FUniformBufferS
 	{
 		for (int32 Index = 0; Index < InUniformBuffers.GetUniformBufferCount(); ++Index)
 		{
-			GlobalUniformBuffers[InUniformBuffers.GetSlot(Index)] = InUniformBuffers.GetUniformBuffer(Index);
+			StaticUniformBuffers[InUniformBuffers.GetSlot(Index)] = InUniformBuffers.GetUniformBuffer(Index);
 		}
 
 		ShaderBindingLayout = nullptr;
@@ -449,14 +575,100 @@ void FVulkanCommandListContext::RHISetStaticUniformBuffers(const FUniformBufferS
 
 void FVulkanCommandListContext::RHISetStaticUniformBuffer(FUniformBufferStaticSlot InSlot, FRHIUniformBuffer* InBuffer)
 {
-	GlobalUniformBuffers[InSlot] = InBuffer;
+	StaticUniformBuffers[InSlot] = InBuffer;
+}
+
+void ApplyStaticUniformBuffersOnContext(
+	FVulkanCommandListContext& Context,
+	FVulkanCommonPipelineDescriptorState& State,
+	const FRHIShader* ShaderRHI,
+	const FVulkanShader* Shader,
+	EShaderFrequency ShaderFrequency,
+	ShaderStage::EStage Stage)
+{
+	if (Shader)
+	{
+		UE::RHICore::ApplyStaticUniformBuffers(
+			ShaderRHI,
+			Context.GetStaticUniformBuffers(),
+			[&Context, &State, Shader, ShaderFrequency, Stage](int32 BufferIndex, FRHIUniformBuffer* Buffer)
+			{
+				BindUniformBuffer(Context, State, Shader, ShaderFrequency, Stage, BufferIndex, ResourceCast(Buffer));
+			}
+		);
+	}
+}
+
+void FVulkanCommandListContext::RHISetGraphicsPipelineState(FRHIGraphicsPipelineState* GraphicsState, uint32 StencilRef, bool bApplyAdditionalState)
+{
+	FVulkanGraphicsPipelineState* Pipeline = ResourceCast(GraphicsState);
+
+	FVulkanPipelineStateCacheManager* PipelineStateCache = Device.GetPipelineStateCache();
+	PipelineStateCache->LRUTouch(Pipeline);
+
+	Pipeline->FrameCounter.Set(GFrameNumberRenderThread);
+
+	FVulkanCommandBuffer& CommandBuffer = GetCommandBuffer();
+	bool bForceResetPipeline = !CommandBuffer.bHasPipeline;
+
+	if (PendingGfxState->SetGfxPipeline(Pipeline, bForceResetPipeline))
+	{
+#if VULKAN_ENABLE_AGGRESSIVE_STATS
+		SCOPE_CYCLE_COUNTER(STAT_VulkanPipelineBind);
+#endif
+		PendingGfxState->Bind(CommandBuffer.GetHandle());
+		CommandBuffer.bHasPipeline = true;
+	}
+
+	PendingGfxState->SetStencilRef(StencilRef);
+
+	if (bApplyAdditionalState)
+	{
+		auto ApplyStaticUniformBuffers = [this, Pipeline](ShaderStage::EStage Stage)
+		{
+			const EShaderFrequency Frequency = ShaderStage::GetFrequencyForGfxStage(Stage);
+			const FVulkanShader* Shader = Pipeline->VulkanShaders[Stage];
+			const FRHIShader* ShaderRHI = RHIShaderFromVulkanShader(Shader, Frequency);
+			ApplyStaticUniformBuffersOnContext(*this, PendingGfxState->GetCurrentState(), ShaderRHI, Shader, Frequency, Stage);
+		};
+
+		ApplyStaticUniformBuffers(ShaderStage::Vertex);
+#if PLATFORM_SUPPORTS_MESH_SHADERS
+		ApplyStaticUniformBuffers(ShaderStage::Mesh);
+		ApplyStaticUniformBuffers(ShaderStage::Task);
+#endif
+#if PLATFORM_SUPPORTS_GEOMETRY_SHADERS
+		ApplyStaticUniformBuffers(ShaderStage::Geometry);
+#endif
+		ApplyStaticUniformBuffers(ShaderStage::Pixel);
+	}
+}
+
+void FVulkanCommandListContext::RHISetComputePipelineState(FRHIComputePipelineState* ComputePipelineState)
+{
+	AcquirePoolSetContainer();
+
+	//#todo-rco: Set PendingGfx to null
+	FVulkanComputePipeline* ComputePipeline = ResourceCast(ComputePipelineState);
+	PendingComputeState->SetComputePipeline(ComputePipeline);
+
+	ComputePipeline->FrameCounter.Set(GFrameNumberRenderThread);
+
+	ApplyStaticUniformBuffersOnContext(
+		*this,
+		PendingComputeState->GetCurrentState(),
+		ComputePipelineState->GetComputeShader(),
+		ResourceCast(ComputePipeline->GetComputeShader()),
+		SF_Compute,
+		ShaderStage::Compute
+	);
 }
 
 void FVulkanCommandListContext::RHISetUniformBufferDynamicOffset(FUniformBufferStaticSlot InSlot, uint32 InOffset)
 {
 	check(IsAligned(InOffset, Device.GetLimits().minUniformBufferOffsetAlignment));
 
-	FVulkanUniformBuffer* UniformBuffer = ResourceCast(GlobalUniformBuffers[InSlot]);
+	FVulkanUniformBuffer* UniformBuffer = ResourceCast(StaticUniformBuffers[InSlot]);
 
 	static const ShaderStage::EStage Stages[2] = 
 	{
@@ -491,95 +703,6 @@ void FVulkanCommandListContext::RHISetUniformBufferDynamicOffset(FUniformBufferS
 	}
 }
 
-static void BindUniformBuffer(
-	FVulkanCommandListContext& Context,
-	FVulkanCommonPipelineDescriptorState& State,
-	FVulkanShader* Shader,
-	EShaderFrequency Frequency,
-	ShaderStage::EStage Stage,
-	uint32 BufferIndex,
-	FVulkanUniformBuffer* UniformBuffer
-)
-{
-#if VULKAN_ENABLE_AGGRESSIVE_STATS
-	SCOPE_CYCLE_COUNTER(STAT_VulkanSetUniformBufferTime);
-#endif
-
-	const FVulkanShaderHeader& CodeHeader = Shader->GetCodeHeader();
-	checkfSlow(!CodeHeader.UniformBufferInfos[BufferIndex].LayoutHash || (CodeHeader.UniformBufferInfos[BufferIndex].LayoutHash == UniformBuffer->GetLayout().GetHash()),
-		TEXT("Mismatched UB layout! Got hash 0x%x, expected 0x%x!"), UniformBuffer->GetLayout().GetHash(), CodeHeader.UniformBufferInfos[BufferIndex].LayoutHash);
-
-	if (UniformBuffer->IsUniformView())
-	{
-		UniformBuffer->SetupUniformBufferView();
-	}
-
-	bool bHasResources = false;
-	if (BufferIndex < CodeHeader.NumBoundUniformBuffers)
-	{
-		checkSlow(UniformBuffer->GetLayout().ConstantBufferSize > 0);
-
-		const VkDescriptorType DescriptorType = State.GetDescriptorType(Stage, BufferIndex);
-
-		if (DescriptorType == VK_DESCRIPTOR_TYPE_UNIFORM_BUFFER_DYNAMIC)
-		{
-			State.SetUniformBuffer<true>(Stage, BufferIndex, UniformBuffer);
-		}
-		else
-		{
-			check(DescriptorType == VK_DESCRIPTOR_TYPE_UNIFORM_BUFFER);
-			State.SetUniformBuffer<false>(Stage, BufferIndex, UniformBuffer);
-		}
-
-		bHasResources = (CodeHeader.UniformBufferInfos[BufferIndex].bHasResources != 0);
-	}
-	else
-	{
-		// If the buffer has no bindings, then it is as resource only ub
-		bHasResources = true;
-	}
-
-	if (bHasResources)
-	{
-		checkSlow(Frequency < SF_NumStandardFrequencies);
-		check(BufferIndex < FVulkanCommandListContext::MAX_UNIFORM_BUFFERS_PER_SHADER_STAGE);
-		Context.BoundUniformBuffers[Frequency][BufferIndex] = UniformBuffer;
-		Context.DirtyUniformBuffers[Frequency] |= (1 << BufferIndex);
-	}
-}
-
-void FVulkanCommandListContext::RHISetShaderUniformBuffer(FRHIGraphicsShader* ShaderRHI, uint32 BufferIndex, FRHIUniformBuffer* BufferRHI)
-{
-	FVulkanShader* Shader = nullptr;
-	const ShaderStage::EStage Stage = GetAndVerifyShaderStageAndVulkanShader(ShaderRHI, PendingGfxState, Shader);
-	check(Shader->GetShaderKey() == PendingGfxState->GetCurrentShaderKey(Stage));
-
-	BindUniformBuffer(
-		*this,
-		PendingGfxState->GetCurrentState(),
-		Shader,
-		Shader->Frequency,
-		Stage,
-		BufferIndex,
-		ResourceCast(BufferRHI)
-	);
-}
-
-void FVulkanCommandListContext::RHISetShaderUniformBuffer(FRHIComputeShader* ComputeShaderRHI, uint32 BufferIndex, FRHIUniformBuffer* BufferRHI)
-{
-	check(PendingComputeState->GetCurrentShader() == ResourceCast(ComputeShaderRHI));
-
-	BindUniformBuffer(
-		*this,
-		PendingComputeState->GetCurrentState(),
-		ResourceCast(ComputeShaderRHI),
-		SF_Compute,
-		ShaderStage::Compute,
-		BufferIndex,
-		ResourceCast(BufferRHI)
-	);
-}
-
 void FVulkanCommandListContext::RHISetStencilRef(uint32 StencilRef)
 {
 	PendingGfxState->SetStencilRef(StencilRef);
diff --git a/Engine/Source/Runtime/VulkanRHI/Private/VulkanContext.cpp b/Engine/Source/Runtime/VulkanRHI/Private/VulkanContext.cpp
--- a/Engine/Source/Runtime/VulkanRHI/Private/VulkanContext.cpp
+++ b/Engine/Source/Runtime/VulkanRHI/Private/VulkanContext.cpp
@@ -270,7 +270,7 @@ FVulkanCommandListContext::FVulkanCommandListContext(FVulkanDevice& InDevice, ER
 	PendingGfxState = new FVulkanPendingGfxState(Device);
 	PendingComputeState = new FVulkanPendingComputeState(Device);
 
-	GlobalUniformBuffers.AddZeroed(FUniformBufferStaticSlotRegistry::Get().GetSlotCount());
+	StaticUniformBuffers.AddZeroed(FUniformBufferStaticSlotRegistry::Get().GetSlotCount());
 }
 
 FVulkanCommandListContext::FVulkanCommandListContext(FVulkanDevice& InDevice, FVulkanParallelRenderPassInfo* InParallelRenderPassInfo)
@@ -285,7 +285,7 @@ FVulkanCommandListContext::FVulkanCommandListContext(FVulkanDevice& InDevice, FV
 	// Only graphic commands can be used
 	PendingGfxState = new FVulkanPendingGfxState(Device);
 
-	GlobalUniformBuffers.AddZeroed(FUniformBufferStaticSlotRegistry::Get().GetSlotCount());
+	StaticUniformBuffers.AddZeroed(FUniformBufferStaticSlotRegistry::Get().GetSlotCount());
 }
 
 void FVulkanCommandListContext::NotifyDeletedRenderTarget(VkImage Image)
@@ -709,10 +709,10 @@ void FVulkanCommandListContext::ApplyShaderBindingLayout(VkShaderStageFlags Shad
 		uint32 UBBindlessIndices[FRHIShaderBindingLayout::MaxUniformBufferEntries];
 		for (uint32 Index = 0; Index < ShaderBindingLayout->GetNumUniformBufferEntries(); ++Index)
 		{
-			checkf(GlobalUniformBuffers[Index], TEXT("Missing static uniform buffer for shader binding layout (%s)."),
+			checkf(StaticUniformBuffers[Index], TEXT("Missing static uniform buffer for shader binding layout (%s)."),
 				*ShaderBindingLayout->GetUniformBufferEntry(Index).LayoutName);
 
-			FVulkanUniformBuffer* StaticUniformBuffer = ResourceCast(GlobalUniformBuffers[Index]);
+			FVulkanUniformBuffer* StaticUniformBuffer = ResourceCast(StaticUniformBuffers[Index]);
 			const FRHIDescriptorHandle BindlessHandle = StaticUniformBuffer->GetBindlessHandle();
 			check(BindlessHandle.IsValid());
 
diff --git a/Engine/Source/Runtime/VulkanRHI/Private/VulkanContext.h b/Engine/Source/Runtime/VulkanRHI/Private/VulkanContext.h
--- a/Engine/Source/Runtime/VulkanRHI/Private/VulkanContext.h
+++ b/Engine/Source/Runtime/VulkanRHI/Private/VulkanContext.h
@@ -269,21 +269,9 @@ public:
 	virtual void RHISetStereoViewport(float LeftMinX, float RightMinX, float LeftMinY, float RightMinY, float MinZ, float LeftMaxX, float RightMaxX, float LeftMaxY, float RightMaxY, float MaxZ) override;
 	virtual void RHISetScissorRect(bool bEnable, uint32 MinX, uint32 MinY, uint32 MaxX, uint32 MaxY) final override;
 	virtual void RHISetGraphicsPipelineState(FRHIGraphicsPipelineState* GraphicsState, uint32 StencilRef, bool bApplyAdditionalState) final override;
-	void RHISetShaderTexture(FRHIGraphicsShader* Shader, uint32 TextureIndex, FRHITexture* NewTexture);
-	void RHISetShaderTexture(FRHIComputeShader* PixelShader, uint32 TextureIndex, FRHITexture* NewTexture);
-	void RHISetShaderSampler(FRHIComputeShader* ComputeShader, uint32 SamplerIndex, FRHISamplerState* NewState);
-	void RHISetShaderSampler(FRHIGraphicsShader* Shader, uint32 SamplerIndex, FRHISamplerState* NewState);
-	void RHISetUAVParameter(FRHIPixelShader* PixelShader, uint32 UAVIndex, FRHIUnorderedAccessView* UAV);
-	void RHISetUAVParameter(FRHIComputeShader* ComputeShader, uint32 UAVIndex, FRHIUnorderedAccessView* UAV);
-	void RHISetShaderResourceViewParameter(FRHIGraphicsShader* Shader, uint32 SamplerIndex, FRHIShaderResourceView* SRV);
-	void RHISetShaderResourceViewParameter(FRHIComputeShader* ComputeShader, uint32 SamplerIndex, FRHIShaderResourceView* SRV);
 	virtual void RHISetStaticUniformBuffers(const FUniformBufferStaticBindings& InUniformBuffers) final override;
 	virtual void RHISetStaticUniformBuffer(FUniformBufferStaticSlot Slot, FRHIUniformBuffer* Buffer) final override;
 	virtual void RHISetUniformBufferDynamicOffset(FUniformBufferStaticSlot Slot, uint32 InOffset) final override;
-	void RHISetShaderUniformBuffer(FRHIGraphicsShader* Shader, uint32 BufferIndex, FRHIUniformBuffer* Buffer);
-	void RHISetShaderUniformBuffer(FRHIComputeShader* ComputeShader, uint32 BufferIndex, FRHIUniformBuffer* Buffer);
-	void RHISetShaderParameter(FRHIGraphicsShader* Shader, uint32 BufferIndex, uint32 BaseIndex, uint32 NumBytes, const void* NewValue);
-	void RHISetShaderParameter(FRHIComputeShader* ComputeShader, uint32 BufferIndex, uint32 BaseIndex, uint32 NumBytes, const void* NewValue);
 	virtual void RHISetShaderParameters(FRHIGraphicsShader* Shader, TConstArrayView<uint8> InParametersData, TConstArrayView<FRHIShaderParameter> InParameters, TConstArrayView<FRHIShaderParameterResource> InResourceParameters, TConstArrayView<FRHIShaderParameterResource> InBindlessParameters) final override;
 	virtual void RHISetShaderParameters(FRHIComputeShader* Shader, TConstArrayView<uint8> InParametersData, TConstArrayView<FRHIShaderParameter> InParameters, TConstArrayView<FRHIShaderParameterResource> InResourceParameters, TConstArrayView<FRHIShaderParameterResource> InBindlessParameters) final override;
 	virtual void RHISetStencilRef(uint32 StencilRef) final override;
@@ -442,16 +430,18 @@ public:
 	TMap<uint32, class FVulkanTypedDescriptorPoolSet*> TypedDescriptorPoolSets;
 	void AcquirePoolSetContainer();
 
+	inline const TArray<FRHIUniformBuffer*>& GetStaticUniformBuffers() const
+	{
+		return StaticUniformBuffers;
+	}
+
 private:
 	void RHIClearMRT(bool bClearColor, int32 NumClearColors, const FLinearColor* ColorArray, bool bClearDepth, float Depth, bool bClearStencil, uint32 Stencil);
 
 	friend class FVulkanDevice;
 	friend class FVulkanDynamicRHI;
 
-	template <typename TRHIShader>
-	void ApplyStaticUniformBuffers(TRHIShader* Shader);
-
-	TArray<FRHIUniformBuffer*> GlobalUniformBuffers;
+	TArray<FRHIUniformBuffer*> StaticUniformBuffers;
 	const FRHIShaderBindingLayout* ShaderBindingLayout = nullptr;
 };
 
diff --git a/Engine/Source/Runtime/VulkanRHI/Private/VulkanPendingState.h b/Engine/Source/Runtime/VulkanRHI/Private/VulkanPendingState.h
--- a/Engine/Source/Runtime/VulkanRHI/Private/VulkanPendingState.h
+++ b/Engine/Source/Runtime/VulkanRHI/Private/VulkanPendingState.h
@@ -213,12 +213,6 @@ public:
 		CurrentPipeline->Bind(CmdBuffer);
 	}
 
-	template<bool bDynamic>
-	void SetUniformBuffer(uint8 DescriptorSet, uint32 BindingIndex, const FVulkanUniformBuffer* UniformBuffer)
-	{
-		CurrentState->SetUniformBuffer<bDynamic>(DescriptorSet, BindingIndex, UniformBuffer);
-	}
-
 	FVulkanPackedUniformBufferState& GetUniformBuffers(uint8 Stage)
 	{
 		return CurrentState->GetUniformBuffers(Stage);
diff --git a/Engine/Source/Runtime/VulkanRHI/Private/VulkanPipelineState.cpp b/Engine/Source/Runtime/VulkanRHI/Private/VulkanPipelineState.cpp
--- a/Engine/Source/Runtime/VulkanRHI/Private/VulkanPipelineState.cpp
+++ b/Engine/Source/Runtime/VulkanRHI/Private/VulkanPipelineState.cpp
@@ -528,81 +528,7 @@ void FVulkanGraphicsPipelineDescriptorState::UpdateBindlessDescriptors(FVulkanCo
 }
 
 
-template <typename TRHIShader>
-void FVulkanCommandListContext::ApplyStaticUniformBuffers(TRHIShader* Shader)
-{
-	if (Shader)
-	{
-		const auto& StaticSlots = Shader->GetStaticSlots();
-		const auto& UBInfos = Shader->GetCodeHeader().UniformBufferInfos;
-
-		for (int32 BufferIndex = 0; BufferIndex < StaticSlots.Num(); ++BufferIndex)
-		{
-			const FUniformBufferStaticSlot Slot = StaticSlots[BufferIndex];
-
-			if (IsUniformBufferStaticSlotValid(Slot))
-			{
-				FRHIUniformBuffer* Buffer = GlobalUniformBuffers[Slot];
-				UE::RHICore::ValidateStaticUniformBuffer(Buffer, Slot, UBInfos[BufferIndex].LayoutHash);
-
-				if (Buffer)
-				{
-					RHISetShaderUniformBuffer(Shader, BufferIndex, Buffer);
-				}
-			}
-		}
-	}
-}
-
-void FVulkanCommandListContext::RHISetGraphicsPipelineState(FRHIGraphicsPipelineState* GraphicsState, uint32 StencilRef, bool bApplyAdditionalState)
-{
-	FVulkanGraphicsPipelineState* Pipeline = ResourceCast(GraphicsState);
-	
-	FVulkanPipelineStateCacheManager* PipelineStateCache = Device.GetPipelineStateCache();
-	PipelineStateCache->LRUTouch(Pipeline);
-
-	Pipeline->FrameCounter.Set(GFrameNumberRenderThread);
-
-	FVulkanCommandBuffer& CommandBuffer = GetCommandBuffer();
-	bool bForceResetPipeline = !CommandBuffer.bHasPipeline;
 
-	if (PendingGfxState->SetGfxPipeline(Pipeline, bForceResetPipeline))
-	{
-#if VULKAN_ENABLE_AGGRESSIVE_STATS
-		SCOPE_CYCLE_COUNTER(STAT_VulkanPipelineBind);
-#endif
-		PendingGfxState->Bind(CommandBuffer.GetHandle());
-		CommandBuffer.bHasPipeline = true;
-	}
-
-	PendingGfxState->SetStencilRef(StencilRef);
-
-	if (bApplyAdditionalState)
-	{
-		ApplyStaticUniformBuffers(static_cast<FVulkanVertexShader*>(Pipeline->VulkanShaders[ShaderStage::Vertex]));
-#if PLATFORM_SUPPORTS_MESH_SHADERS
-		ApplyStaticUniformBuffers(static_cast<FVulkanMeshShader*>(Pipeline->VulkanShaders[ShaderStage::Mesh]));
-		ApplyStaticUniformBuffers(static_cast<FVulkanTaskShader*>(Pipeline->VulkanShaders[ShaderStage::Task]));
-#endif
-#if PLATFORM_SUPPORTS_GEOMETRY_SHADERS
-		ApplyStaticUniformBuffers(static_cast<FVulkanGeometryShader*>(Pipeline->VulkanShaders[ShaderStage::Geometry]));
-#endif
-		ApplyStaticUniformBuffers(static_cast<FVulkanPixelShader*>(Pipeline->VulkanShaders[ShaderStage::Pixel]));
-	}
-}
-
-void FVulkanCommandListContext::RHISetComputePipelineState(FRHIComputePipelineState* ComputePipelineState)
-{
-	AcquirePoolSetContainer();
-
-	//#todo-rco: Set PendingGfx to null
-	FVulkanComputePipeline* ComputePipeline = ResourceCast(ComputePipelineState);
-	PendingComputeState->SetComputePipeline(ComputePipeline);
-
-	ComputePipeline->FrameCounter.Set(GFrameNumberRenderThread);
-	
-	ApplyStaticUniformBuffers(ResourceCast(ComputePipeline->GetComputeShader()));
-}
 
 
 template bool FVulkanGraphicsPipelineDescriptorState::InternalUpdateDescriptorSets<true>(FVulkanCommandListContext& Context);
diff --git a/Engine/Source/Runtime/VulkanRHI/Private/VulkanPipelineState.h b/Engine/Source/Runtime/VulkanRHI/Private/VulkanPipelineState.h
--- a/Engine/Source/Runtime/VulkanRHI/Private/VulkanPipelineState.h
+++ b/Engine/Source/Runtime/VulkanRHI/Private/VulkanPipelineState.h
@@ -94,8 +94,7 @@ public:
 		MarkDirty(DSWriter[DescriptorSet].WriteInputAttachment(BindingIndex, TextureView, Layout));
 	}
 
-	template<bool bDynamic>
-	inline void SetUniformBuffer(uint8 DescriptorSet, uint32 BindingIndex, const FVulkanUniformBuffer* UniformBuffer)
+	inline void SetUniformBuffer(uint8 DescriptorSet, uint32 BindingIndex, const FVulkanUniformBuffer* UniformBuffer, bool bDynamic)
 	{
 		const VulkanRHI::FVulkanAllocation& Allocation = UniformBuffer->GetAllocation();
 		VkDeviceSize Range = UniformBuffer->IsUniformView() ? PLATFORM_MAX_UNIFORM_BUFFER_RANGE : UniformBuffer->GetSize();
