VulkanRHI: Allow for enabling the dynamic state extensions individually.

#rnx

#Synced-CL 48355128

diff --git a/Engine/Source/Runtime/VulkanRHI/Private/VulkanExtensions.cpp b/Engine/Source/Runtime/VulkanRHI/Private/VulkanExtensions.cpp
--- a/Engine/Source/Runtime/VulkanRHI/Private/VulkanExtensions.cpp
+++ b/Engine/Source/Runtime/VulkanRHI/Private/VulkanExtensions.cpp
@@ -85,9 +85,15 @@ TAutoConsoleVariable<int32> GVulkanAllowDynamicRenderingCVar(
 TAutoConsoleVariable<int32> GVulkanAllowDynamicStatesCVar(
 	TEXT("r.Vulkan.AllowDynamicStates"),
 	0,
-	TEXT("Enables the use of VK_EXT_extended_dynamic_state1,2,3 to reduce amount of unique pipelines. \n")
-	TEXT("0: Do not enable support for all the dynamic state extensions (default)\n")
-	TEXT("1: Enable all the dynamic state extensions"),
+	TEXT("Enables the use of VK_EXT_extended_dynamic_state1, 2, 3 and \n")
+	TEXT("VK_EXT_vertex_input_dynamic_state to reduce amount of unique pipelines. \n")
+	TEXT("0: Do not enable support for a dynamic state extensions (default)\n")
+	TEXT("1: Enable only VK_EXT_extended_dynamic_state\n")
+	TEXT("2: Enable only VK_EXT_extended_dynamic_state2\n")
+	TEXT("4: Enable only VK_EXT_extended_dynamic_state3\n")
+	TEXT("8: Enable only VK_EXT_vertex_input_dynamic_state\n")
+	TEXT("7: Enable VK_EXT_extended_dynamic_state1,2,3.\n")
+	TEXT("15: Enable all dynamic state extensions.\n"),
 	ECVF_ReadOnly
 );
 
@@ -2052,7 +2058,7 @@ public:
 	FVulkanEXTExtendedDynamicStateExtension(FVulkanDevice* InDevice)
 		: FVulkanDeviceExtension(InDevice, VK_EXT_EXTENDED_DYNAMIC_STATE_EXTENSION_NAME, VULKAN_EXTENSION_ENABLED, VK_API_VERSION_1_3)
 	{
-		bEnabledInCode = bEnabledInCode && (GVulkanAllowDynamicStatesCVar.GetValueOnAnyThread() != 0);
+		bEnabledInCode = bEnabledInCode && ((GVulkanAllowDynamicStatesCVar.GetValueOnAnyThread() & 1) != 0);
 	}
 
 	virtual void PrePhysicalDeviceFeatures(VkPhysicalDeviceFeatures2KHR& PhysicalDeviceFeatures2) override final
@@ -2086,7 +2092,7 @@ public:
 	FVulkanEXTExtendedDynamicState2Extension(FVulkanDevice* InDevice)
 		: FVulkanDeviceExtension(InDevice, VK_EXT_EXTENDED_DYNAMIC_STATE_2_EXTENSION_NAME, VULKAN_EXTENSION_ENABLED, VK_API_VERSION_1_3)
 	{
-		bEnabledInCode = bEnabledInCode && (GVulkanAllowDynamicStatesCVar.GetValueOnAnyThread() != 0);
+		bEnabledInCode = bEnabledInCode && ((GVulkanAllowDynamicStatesCVar.GetValueOnAnyThread() & 2) != 0);
 	}
 
 	virtual void PrePhysicalDeviceFeatures(VkPhysicalDeviceFeatures2KHR& PhysicalDeviceFeatures2) override final
@@ -2120,7 +2126,7 @@ public:
 	FVulkanEXTExtendedDynamicState3Extension(FVulkanDevice* InDevice)
 		: FVulkanDeviceExtension(InDevice, VK_EXT_EXTENDED_DYNAMIC_STATE_3_EXTENSION_NAME, VULKAN_EXTENSION_ENABLED)
 	{
-		bEnabledInCode = bEnabledInCode && (GVulkanAllowDynamicStatesCVar.GetValueOnAnyThread() != 0);
+		bEnabledInCode = bEnabledInCode && ((GVulkanAllowDynamicStatesCVar.GetValueOnAnyThread() & 4) != 0);
 	}
 
 	virtual void PrePhysicalDeviceFeatures(VkPhysicalDeviceFeatures2KHR& PhysicalDeviceFeatures2) override final
@@ -2160,7 +2166,7 @@ public:
 	FVulkanEXTVertexInputDynamicStateExtension(FVulkanDevice* InDevice)
 		: FVulkanDeviceExtension(InDevice, VK_EXT_VERTEX_INPUT_DYNAMIC_STATE_EXTENSION_NAME, VULKAN_EXTENSION_ENABLED)
 	{
-		bEnabledInCode = bEnabledInCode && (GVulkanAllowDynamicStatesCVar.GetValueOnAnyThread() != 0);
+		bEnabledInCode = bEnabledInCode && ((GVulkanAllowDynamicStatesCVar.GetValueOnAnyThread() & 8) != 0);
 	}
 
 	virtual void PrePhysicalDeviceFeatures(VkPhysicalDeviceFeatures2KHR& PhysicalDeviceFeatures2) override final
