Enable vulkan VK_EXT_image_robustness extension by default.
#tests local client testing.
#Synced-CL 48498858

#Synced-CL 46813550

diff --git a/Engine/Source/Runtime/VulkanRHI/Private/VulkanDevice.h b/Engine/Source/Runtime/VulkanRHI/Private/VulkanDevice.h
--- a/Engine/Source/Runtime/VulkanRHI/Private/VulkanDevice.h
+++ b/Engine/Source/Runtime/VulkanRHI/Private/VulkanDevice.h
@@ -107,6 +107,7 @@ struct FOptionalVulkanDeviceExtensions
 			uint64 HasEXTExtendedDynamicState2 : 1;
 			uint64 HasEXTExtendedDynamicState3 : 1;
 			uint64 HasEXTVertexInputDynamicState : 1;
+			uint64 HasEXTImageRobustness : 1;
 
 			// Promoted to 1.4
 			uint64 HasKHRMaintenance5 : 1;
diff --git a/Engine/Source/Runtime/VulkanRHI/Private/VulkanExtensions.cpp b/Engine/Source/Runtime/VulkanRHI/Private/VulkanExtensions.cpp
--- a/Engine/Source/Runtime/VulkanRHI/Private/VulkanExtensions.cpp
+++ b/Engine/Source/Runtime/VulkanRHI/Private/VulkanExtensions.cpp
@@ -146,6 +146,15 @@ TAutoConsoleVariable<int32> GVulkanVariableRateShadingFormatCVar(
 	ECVF_ReadOnly
 );
 
+TAutoConsoleVariable<int32> GVulkanAllowImageRobustness(
+	TEXT("r.Vulkan.AllowImageRobustness"),
+	1,
+	TEXT("Whether to allow VK_EXT_image_robustness\n")
+	TEXT("0 to disable\n")
+	TEXT("1 to enable (default)"),
+	ECVF_ReadOnly
+);
+
 #if VULKAN_HAS_DEBUGGING_ENABLED
 extern TAutoConsoleVariable<int32> GGPUValidationCvar;
 #endif
@@ -1747,6 +1756,40 @@ public:
 	}
 };
 
+// ***** VK_EXT_image_robustness
+class FVulkanEXTImageRobustnessExtension : public FVulkanDeviceExtension
+{
+public:
+
+	FVulkanEXTImageRobustnessExtension(FVulkanDevice* InDevice)
+		: FVulkanDeviceExtension(InDevice, VK_EXT_IMAGE_ROBUSTNESS_EXTENSION_NAME, VULKAN_EXTENSION_ENABLED)
+	{
+		bEnabledInCode = bEnabledInCode && (GVulkanAllowImageRobustness.GetValueOnAnyThread() != 0);
+	}
+
+	virtual void PrePhysicalDeviceFeatures(VkPhysicalDeviceFeatures2KHR& PhysicalDeviceFeatures2) override final
+	{
+		ZeroVulkanStruct(ImageRobustnessFeature, VK_STRUCTURE_TYPE_PHYSICAL_DEVICE_IMAGE_ROBUSTNESS_FEATURES_EXT);
+		AddToPNext(PhysicalDeviceFeatures2, ImageRobustnessFeature);
+	}
+
+	virtual void PostPhysicalDeviceFeatures(FOptionalVulkanDeviceExtensions& ExtensionFlags) override final
+	{
+		bRequirementsPassed = (ImageRobustnessFeature.robustImageAccess == VK_TRUE);
+		ExtensionFlags.HasEXTImageRobustness = bRequirementsPassed;
+	}
+
+	virtual void PreCreateDevice(VkDeviceCreateInfo& DeviceCreateInfo) override final
+	{
+		if (bRequirementsPassed)
+		{
+			AddToPNext(DeviceCreateInfo, ImageRobustnessFeature);
+		}
+	}
+
+	VkPhysicalDeviceImageRobustnessFeaturesEXT ImageRobustnessFeature;
+};
+
 // ***** VK_EXT_pipeline_creation_cache_control
 class FVulkanEXTPipelineCreationCacheControlExtension : public FVulkanDeviceExtension
 {
@@ -2483,6 +2526,7 @@ FVulkanDeviceExtensionArray FVulkanDeviceExtension::GetUESupportedDeviceExtensio
 	ADD_CUSTOM_EXTENSION(FVulkanKHR16BitStorageExtension);
 	ADD_CUSTOM_EXTENSION(FVulkanKHRShaderFloat16Int8Extension);
 	ADD_CUSTOM_EXTENSION(FVulkanKHRDepthStencilResolveExtension);
+	ADD_CUSTOM_EXTENSION(FVulkanEXTImageRobustnessExtension);
 	ADD_CUSTOM_EXTENSION(FVulkanEXTPipelineCreationCacheControlExtension);
 	ADD_CUSTOM_EXTENSION(FVulkanKHRFragmentShaderBarycentricExtension);
 	ADD_CUSTOM_EXTENSION(FVulkanKHRComputeShaderDerivatives);
