VulkanRHI: Fix memory getting freed too early on Android for certain textures with external owner.

#rnx

#Synced-CL 48947181

diff --git a/Engine/Source/Runtime/VulkanRHI/Private/VulkanTexture.cpp b/Engine/Source/Runtime/VulkanRHI/Private/VulkanTexture.cpp
--- a/Engine/Source/Runtime/VulkanRHI/Private/VulkanTexture.cpp
+++ b/Engine/Source/Runtime/VulkanRHI/Private/VulkanTexture.cpp
@@ -1995,6 +1995,7 @@ FVulkanTexture::FVulkanTexture(FVulkanDevice& InDevice, const FRHITextureCreateD
 
 struct FVulkanAndroidTextureResources
 {
+	FVulkanDevice* Device;
 	VkImage Image;
 	VkDeviceMemory DeviceMemory;
 	VkSamplerYcbcrConversion SamplerYcbcrConversion;
@@ -2004,34 +2005,36 @@ struct FVulkanAndroidTextureResources
 static void CleanupVulkanAndroidTextureResources(void* UserData)
 {
 	check(UserData);
-
 	FVulkanAndroidTextureResources* VulkanResources = static_cast<FVulkanAndroidTextureResources*>(UserData);
 
-	IVulkanDynamicRHI* RHI = GetIVulkanDynamicRHI();
-	VkDevice Device = RHI->RHIGetVkDevice();
-	const VkAllocationCallbacks* AllocationCallbacks = RHI->RHIGetVkAllocationCallbacks();
-
-	if (VulkanResources->SamplerYcbcrConversion != VK_NULL_HANDLE)
+	if (VulkanResources->Device)
 	{
-		VulkanRHI::vkDestroySamplerYcbcrConversion(Device, VulkanResources->SamplerYcbcrConversion, AllocationCallbacks);
-	}
+		VulkanResources->Device->GetDeferredDeletionQueue().EnqueueFunc([VulkanResources](FVulkanDevice& Device)
+		{
+			if (VulkanResources->SamplerYcbcrConversion != VK_NULL_HANDLE)
+			{
+				VulkanRHI::vkDestroySamplerYcbcrConversion(Device.GetHandle(), VulkanResources->SamplerYcbcrConversion, VULKAN_CPU_ALLOCATOR);
+			}
 
-	if (VulkanResources->DeviceMemory != VK_NULL_HANDLE)
-	{
-		VulkanRHI::vkFreeMemory(Device, VulkanResources->DeviceMemory, AllocationCallbacks);
-	}
+			if (VulkanResources->DeviceMemory != VK_NULL_HANDLE)
+			{
+				VulkanRHI::vkFreeMemory(Device.GetHandle(), VulkanResources->DeviceMemory, VULKAN_CPU_ALLOCATOR);
+			}
 
-	if (VulkanResources->Image != VK_NULL_HANDLE)
-	{
-		VulkanRHI::vkDestroyImage(Device, VulkanResources->Image, AllocationCallbacks);
-	}
+			if (VulkanResources->Image != VK_NULL_HANDLE)
+			{
+				VulkanRHI::vkDestroyImage(Device.GetHandle(), VulkanResources->Image, VULKAN_CPU_ALLOCATOR);
+			}
 
-	if (VulkanResources->HardwareBuffer)
-	{
-		AHardwareBuffer_release(VulkanResources->HardwareBuffer);
+			if (VulkanResources->HardwareBuffer)
+			{
+				AHardwareBuffer_release(VulkanResources->HardwareBuffer);
+			}
+
+			delete VulkanResources;
+		});
 	}
 
-	delete VulkanResources;
 }
 
 FVulkanTexture::FVulkanTexture(FVulkanDevice& InDevice, const FRHITextureCreateDesc& InCreateDesc, const AHardwareBuffer_Desc& HardwareBufferDesc, AHardwareBuffer* HardwareBuffer)
@@ -2044,9 +2047,7 @@ FVulkanTexture::FVulkanTexture(FVulkanDevice& InDevice, const FRHITextureCreateD
 	check(HardwareBuffer);
 	AHardwareBuffer_acquire(HardwareBuffer);
 
-	IVulkanDynamicRHI* RHI = GetIVulkanDynamicRHI();
 	VkDevice VulkanDevice = InDevice.GetHandle();
-	const VkAllocationCallbacks* AllocationCallbacks = RHI->RHIGetVkAllocationCallbacks();
 
 	VkAndroidHardwareBufferFormatPropertiesANDROID HardwareBufferFormatProperties;
 	ZeroVulkanStruct(HardwareBufferFormatProperties, VK_STRUCTURE_TYPE_ANDROID_HARDWARE_BUFFER_FORMAT_PROPERTIES_ANDROID);
@@ -2089,7 +2090,7 @@ FVulkanTexture::FVulkanTexture(FVulkanDevice& InDevice, const FRHITextureCreateD
 	ImageCreateInfo.pNext = &ExternalMemoryImageCreateInfo;
 
 	VkImage VulkanImage;
-	VERIFYVULKANRESULT(VulkanRHI::vkCreateImage(VulkanDevice, &ImageCreateInfo, AllocationCallbacks, &VulkanImage));
+	VERIFYVULKANRESULT(VulkanRHI::vkCreateImage(VulkanDevice, &ImageCreateInfo, VULKAN_CPU_ALLOCATOR, &VulkanImage));
 
 	VkMemoryDedicatedAllocateInfo MemoryDedicatedAllocateInfo;
 	ZeroVulkanStruct(MemoryDedicatedAllocateInfo, VK_STRUCTURE_TYPE_MEMORY_DEDICATED_ALLOCATE_INFO);
@@ -2116,7 +2117,7 @@ FVulkanTexture::FVulkanTexture(FVulkanDevice& InDevice, const FRHITextureCreateD
 	MemoryAllocateInfo.pNext = &ImportAndroidHardwareBufferInfo;
 
 	VkDeviceMemory VulkanDeviceMemory;
-	VERIFYVULKANRESULT(VulkanRHI::vkAllocateMemory(VulkanDevice, &MemoryAllocateInfo, AllocationCallbacks, &VulkanDeviceMemory));
+	VERIFYVULKANRESULT(VulkanRHI::vkAllocateMemory(VulkanDevice, &MemoryAllocateInfo, VULKAN_CPU_ALLOCATOR, &VulkanDeviceMemory));
 	VERIFYVULKANRESULT(VulkanRHI::vkBindImageMemory(VulkanDevice, VulkanImage, VulkanDeviceMemory, 0));
 
 	VkSamplerYcbcrConversionCreateInfo SamplerYcbcrConversionCreateInfo;
@@ -2132,10 +2133,11 @@ FVulkanTexture::FVulkanTexture(FVulkanDevice& InDevice, const FRHITextureCreateD
 	SamplerYcbcrConversionCreateInfo.pNext = &ExternalFormat;
 
 	VkSamplerYcbcrConversion SamplerYcbcrConversion;
-	VERIFYVULKANRESULT(VulkanRHI::vkCreateSamplerYcbcrConversion(VulkanDevice, &SamplerYcbcrConversionCreateInfo, AllocationCallbacks, &SamplerYcbcrConversion));
+	VERIFYVULKANRESULT(VulkanRHI::vkCreateSamplerYcbcrConversion(VulkanDevice, &SamplerYcbcrConversionCreateInfo, VULKAN_CPU_ALLOCATOR, &SamplerYcbcrConversion));
 
 	FVulkanAndroidTextureResources* VulkanAndroidTextureResources = new FVulkanAndroidTextureResources
 	{
+		&InDevice,
 		VulkanImage,
 		VulkanDeviceMemory,
 		SamplerYcbcrConversion,
