VulkanRHI: Change container for payloads (can iterate on queued payloads)

#rnx

#Synced-CL 49034046

diff --git a/Engine/Source/Runtime/VulkanRHI/Private/VulkanQueue.cpp b/Engine/Source/Runtime/VulkanRHI/Private/VulkanQueue.cpp
--- a/Engine/Source/Runtime/VulkanRHI/Private/VulkanQueue.cpp
+++ b/Engine/Source/Runtime/VulkanRHI/Private/VulkanQueue.cpp
@@ -225,8 +225,7 @@ int32 FVulkanQueue::SubmitQueuedPayloads(TMap<VkSemaphore, FBinarySemaphoreSigna
 
 	// Accumulate a list of the payloads we can submit
 	TArray<FVulkanPayload*> Payloads;
-	FVulkanPayload* Payload = nullptr;
-	while (PendingSubmission.Peek(Payload))
+	while (FVulkanPayload * Payload = PendingSubmission.Peek())
 	{
 		// We can only submit the payload if all its wait semas have been signaled
 		if (!CanBeProcessed(Payload))
@@ -509,12 +508,11 @@ int32 FVulkanQueue::ProcessInterruptQueue(uint64 Timeout)
 	}
 
 	int32 NumPayloads = 0;
-	FVulkanPayload* Payload = nullptr;
 	bool bSuccess;
 	do
 	{
 		bSuccess = false;
-		if (PendingInterrupt.Peek(Payload))
+		if (FVulkanPayload* Payload = PendingInterrupt.Peek())
 		{
 			if (Payload->TimelineSemaphoreValue <= CompletedTimelineSemaphoreValue)
 			{
@@ -543,19 +541,18 @@ int32 FVulkanQueue::ProcessInterruptQueue(uint64 Timeout)
 					CompletedTimelineSemaphoreValue = Payload->TimelineSemaphoreValue;
 				}
 			}
-		}
 
-		if (bSuccess)
-		{
-			++NumPayloads;
+			if (bSuccess)
+			{
+				++NumPayloads;
 
-			// Resolve any pending actions from the payload being completed
-			FVulkanDynamicRHI::Get().CompletePayload(Payload);
+				// Resolve any pending actions from the payload being completed
+				FVulkanDynamicRHI::Get().CompletePayload(Payload);
 
-			PendingInterrupt.Pop();
-			delete Payload;
+				PendingInterrupt.Pop();
+				delete Payload;
+			}
 		}
-
 	} while (bSuccess);
 
 	return NumPayloads;
diff --git a/Engine/Source/Runtime/VulkanRHI/Private/VulkanQueue.h b/Engine/Source/Runtime/VulkanRHI/Private/VulkanQueue.h
--- a/Engine/Source/Runtime/VulkanRHI/Private/VulkanQueue.h
+++ b/Engine/Source/Runtime/VulkanRHI/Private/VulkanQueue.h
@@ -103,11 +103,6 @@ public:
 		return CompletedTimelineSemaphoreValue;
 	}
 
-	void EnqueuePayload(FVulkanPayload* InPayload)
-	{
-		PendingSubmission.Enqueue(InPayload);
-	}
-
 	int32 SubmitQueuedPayloads(TMap<VkSemaphore, FBinarySemaphoreSignalInfo>& SignaledSemas);
 	int32 ProcessInterruptQueue(uint64 Timeout);
 
@@ -125,13 +120,6 @@ private:
 
 	void BindSparseResources(FVulkanPayload& Payload);
 
-	// Used by submission pipe which holds the proper locks to access this queue
-	FVulkanPayload* GetNextInterruptPayload()
-	{
-		FVulkanPayload* Payload = nullptr;
-		PendingInterrupt.Peek(Payload);
-		return Payload;
-	}
 	friend FVulkanDynamicRHI;
 
 	VkQueue Queue;
@@ -152,8 +140,38 @@ private:
 	VkPipelineStageFlags SupportedStages = VK_PIPELINE_STAGE_NONE;
 	VkAccessFlags SupportedAccess = VK_ACCESS_NONE;
 
-	TQueue<FVulkanPayload*, EQueueMode::Mpsc> PendingSubmission;
-	TQueue<FVulkanPayload*, EQueueMode::Spsc> PendingInterrupt;
+	class FVulkanPayloadQueue
+	{
+		TSpscQueue<FVulkanPayload*> Queue;
+
+	public:
+		void Enqueue(FVulkanPayload* Payload)
+		{
+			check(Payload);
+			Queue.Enqueue(Payload);
+		}
+
+		FVulkanPayload* Peek()
+		{
+			if (FVulkanPayload** PayloadPtr = Queue.Peek())
+			{
+				return *PayloadPtr;
+			}
+
+			return nullptr;
+		}
+
+		void Pop()
+		{
+			Queue.Dequeue();
+		}
+
+		auto begin() { return Queue.begin(); }
+		auto end  () { return Queue.end  (); }
+	};
+
+	FVulkanPayloadQueue PendingSubmission;
+	FVulkanPayloadQueue PendingInterrupt;
 
 	// The active timing struct on this queue. Updated / accessed by the interrupt thread.
 	FVulkanTiming* Timing = nullptr;
diff --git a/Engine/Source/Runtime/VulkanRHI/Private/VulkanSubmission.cpp b/Engine/Source/Runtime/VulkanRHI/Private/VulkanSubmission.cpp
--- a/Engine/Source/Runtime/VulkanRHI/Private/VulkanSubmission.cpp
+++ b/Engine/Source/Runtime/VulkanRHI/Private/VulkanSubmission.cpp
@@ -340,7 +340,7 @@ bool FVulkanDynamicRHI::ProcessSubmissionQueue()
 	while (PendingPayloadsForSubmission.Dequeue(PlatformCmdList))
 	{
 	#if WITH_RHI_BREADCRUMBS
-		TSharedPtr<FRHIBreadcrumbAllocatorArray> BreadcrumbAllocators {};
+		TSharedPtr<FRHIBreadcrumbAllocatorArray> BreadcrumbAllocators{};
 		if (PlatformCmdList->BreadcrumbAllocators.Num())
 		{
 			BreadcrumbAllocators = MakeShared<FRHIBreadcrumbAllocatorArray>(MoveTemp(PlatformCmdList->BreadcrumbAllocators));
@@ -358,7 +358,7 @@ bool FVulkanDynamicRHI::ProcessSubmissionQueue()
 			}
 		#endif
 
-			Payload->Queue.EnqueuePayload(Payload);
+			Payload->Queue.PendingSubmission.Enqueue(Payload);
 		}
 
 		delete PlatformCmdList;
@@ -399,8 +399,7 @@ bool FVulkanDynamicRHI::WaitAndProcessInterruptQueue()
 			bProgress |= (Queue.ProcessInterruptQueue(0) > 0);
 
 			// Get the next payload to wait on
-			FVulkanPayload* Payload = Queue.GetNextInterruptPayload();
-			if (Payload)
+			if (FVulkanPayload* Payload = Queue.PendingInterrupt.Peek())
 			{
 				NextPayloads[NumSyncs++] = Payload;
 			}
@@ -722,85 +721,85 @@ void FVulkanDynamicRHI::CompletePayload(FVulkanPayload* Payload)
 
 			for (uint32 QueryIndex = 0; QueryIndex < QueryPool->CurrentQueryCount; ++QueryIndex)
 			{
-				uint64* ResultPtr = QueryPool->QueryResults[QueryIndex];
+					uint64* ResultPtr = QueryPool->QueryResults[QueryIndex];
 				if (!ResultPtr)
-				{
+					{
 					continue;
 				}
 
-				uint64 QueryResult = 0;
-				for (uint32 ViewIndex = 0; ViewIndex < QueryPool->QueryStride; ++ViewIndex)
-				{
-					uint32 ResultIndex = QueryIndex * QueryPool->QueryStride + ViewIndex;
-					if (TempResults[ResultIndex].Availability)
+					uint64 QueryResult = 0;
+					for (uint32 ViewIndex = 0; ViewIndex < QueryPool->QueryStride; ++ViewIndex)
 					{
-						if (bUseMaxValue)
+						uint32 ResultIndex = QueryIndex * QueryPool->QueryStride + ViewIndex;
+						if (TempResults[ResultIndex].Availability)
 						{
-							QueryResult = FMath::Max(QueryResult, TempResults[ResultIndex].Result);
+							if (bUseMaxValue)
+							{
+								QueryResult = FMath::Max(QueryResult, TempResults[ResultIndex].Result);
+							}
+							else
+							{
+								QueryResult += TempResults[ResultIndex].Result;
+							}
 						}
 						else
 						{
-							QueryResult += TempResults[ResultIndex].Result;
+							UE_LOG(LogVulkanRHI, Warning, TEXT("Could not read back query: [PoolType=%d, Index=%d, View=%d]"), (int32)QueryPool->GetPoolType(), QueryIndex, ViewIndex);
 						}
 					}
-					else
-					{
-						UE_LOG(LogVulkanRHI, Warning, TEXT("Could not read back query: [PoolType=%d, Index=%d, View=%d]"), (int32)QueryPool->GetPoolType(), QueryIndex, ViewIndex);
-					}
-				}
 
-				if (QueryPool->GetPoolType() == EVulkanQueryPoolType::Timestamp)
-				{
-					// Convert from GPU timestamp to CPU timestamp (relative to FPlatformTime::Cycles64())
-					checkSlow(Payload->Queue.Timing);
-					const FVulkanTiming* CurrentTiming = Payload->Queue.Timing;
+						if (QueryPool->GetPoolType() == EVulkanQueryPoolType::Timestamp)
+						{
+							// Convert from GPU timestamp to CPU timestamp (relative to FPlatformTime::Cycles64())
+							checkSlow(Payload->Queue.Timing);
+							const FVulkanTiming* CurrentTiming = Payload->Queue.Timing;
 
-					uint64 CPUCalibrationTimestamp = CurrentTiming->CPUTimestamp;
-					uint64 GPUCalibrationTimestamp = CurrentTiming->GPUTimestamp;
+							uint64 CPUCalibrationTimestamp = CurrentTiming->CPUTimestamp;
+							uint64 GPUCalibrationTimestamp = CurrentTiming->GPUTimestamp;
 
-					// NOTE: Adreno GPU calibration values aren't valid, create a fallback
-					if (bBypassTiming || (QueryResult < GPUCalibrationTimestamp) || (GPUCalibrationTimestamp == 0))
-					{
-						if (!bBypassTiming)
-						{
-							static bool bWarnOnce = false;
-							if (!bWarnOnce)
+							// NOTE: Adreno GPU calibration values aren't valid, create a fallback
+						if (bBypassTiming || (QueryResult < GPUCalibrationTimestamp) || (GPUCalibrationTimestamp == 0))
 							{
-								bWarnOnce = true;
-								UE_LOG(LogVulkanRHI, Warning,
-									TEXT("GPU timestamp can't be moved to CPU timeline.  ")
-									TEXT("Calibration data is invalid (Curr=%lld  Cal=%lld)!  ")
-									TEXT("CPU and GPU profiling data will not line up!"),
-								QueryResult, GPUCalibrationTimestamp);
-							}
-
-							const uint64 CurrentTimestamp = FPlatformTime::Cycles64();
+								if (!bBypassTiming)
+								{
+									static bool bWarnOnce = false;
+									if (!bWarnOnce)
+									{
+										bWarnOnce = true;
+										UE_LOG(LogVulkanRHI, Warning,
+											TEXT("GPU timestamp can't be moved to CPU timeline.  ")
+											TEXT("Calibration data is invalid (Curr=%lld  Cal=%lld)!  ")
+											TEXT("CPU and GPU profiling data will not line up!"),
+										QueryResult, GPUCalibrationTimestamp);
+									}
+
+									const uint64 CurrentTimestamp = FPlatformTime::Cycles64();
 
 							ReferenceCPUTimestamp = ((CPUCalibrationTimestamp != 0) && (CPUCalibrationTimestamp < CurrentTimestamp))
 								? CPUCalibrationTimestamp
 								: CurrentTimestamp;
 									
-							ReferenceGPUTimestamp = QueryResult;
-							bBypassTiming = true;
-						}
+								ReferenceGPUTimestamp = QueryResult;
+									bBypassTiming = true;
+								}
 
-						CPUCalibrationTimestamp = ReferenceCPUTimestamp;
-						GPUCalibrationTimestamp = ReferenceGPUTimestamp;
+								CPUCalibrationTimestamp = ReferenceCPUTimestamp;
+								GPUCalibrationTimestamp = ReferenceGPUTimestamp;
 
-						// :todo: Clamp results within the imperfect calibration timestamps (to prevent negative values below)
-						QueryResult = FMath::Max(QueryResult, GPUCalibrationTimestamp);
-					}
+								// :todo: Clamp results within the imperfect calibration timestamps (to prevent negative values below)
+							QueryResult = FMath::Max(QueryResult, GPUCalibrationTimestamp);
+							}
 
-					const uint64 GPUDelta = QueryResult - GPUCalibrationTimestamp;
-					const uint64 CPUDelta = (GPUDelta * CurrentTiming->CPUFrequency) / CurrentTiming->GPUFrequency;
+						const uint64 GPUDelta = QueryResult - GPUCalibrationTimestamp;
+							const uint64 CPUDelta = (GPUDelta * CurrentTiming->CPUFrequency) / CurrentTiming->GPUFrequency;
 
-					(*ResultPtr) = CPUCalibrationTimestamp + CPUDelta;
-				}
-				else
-				{
-					(*ResultPtr) = QueryResult;
-				}
-			}
+							(*ResultPtr) = CPUCalibrationTimestamp + CPUDelta;
+						}
+						else
+						{
+						(*ResultPtr) = QueryResult;
+						}
+					}
 
 			QueryPool->QueryRefs.Empty(QueryPool->GetMaxQueries());
 			Device->ReleaseQueryPool(QueryPool);
diff --git a/Engine/Source/Runtime/VulkanRHI/Private/VulkanUtil.cpp b/Engine/Source/Runtime/VulkanRHI/Private/VulkanUtil.cpp
--- a/Engine/Source/Runtime/VulkanRHI/Private/VulkanUtil.cpp
+++ b/Engine/Source/Runtime/VulkanRHI/Private/VulkanUtil.cpp
@@ -604,27 +604,27 @@ void FVulkanDynamicRHI::DumpCrashMarkers()
 		};
 
 		Device->ForEachQueue([](FVulkanQueue& Queue)
+		{
+			TArray<FCheckpointDataNV> Checkpoints;
+			uint32 Num = 0;
+			const VkQueue QueueHandle = Queue.GetHandle();
+			VulkanDynamicAPI::vkGetQueueCheckpointDataNV(QueueHandle, &Num, nullptr);
+			if (Num > 0)
 			{
-				TArray<FCheckpointDataNV> Checkpoints;
-				uint32 Num = 0;
-				const VkQueue QueueHandle = Queue.GetHandle();
-				VulkanDynamicAPI::vkGetQueueCheckpointDataNV(QueueHandle, &Num, nullptr);
-				if (Num > 0)
+				Checkpoints.AddDefaulted(Num);
+				VulkanDynamicAPI::vkGetQueueCheckpointDataNV(QueueHandle, &Num, &Checkpoints[0]);
+				check(Num == Checkpoints.Num());
+				for (uint32 Index = 0; Index < Num; ++Index)
 				{
-					Checkpoints.AddDefaulted(Num);
-					VulkanDynamicAPI::vkGetQueueCheckpointDataNV(QueueHandle, &Num, &Checkpoints[0]);
-					check(Num == Checkpoints.Num());
-					for (uint32 Index = 0; Index < Num; ++Index)
-					{
-						check(Checkpoints[Index].sType == VK_STRUCTURE_TYPE_CHECKPOINT_DATA_NV);
-						FRHIBreadcrumbNode* Breadcrumb = (FRHIBreadcrumbNode*)Checkpoints[Index].pCheckpointMarker;
-						UE_LOG(LogVulkanRHI, Error, TEXT("[VK_NV_device_diagnostic_checkpoints] %i: Stage %s (0x%08x), %s"),
-							Index, VK_TYPE_TO_STRING(VkPipelineStageFlagBits, Checkpoints[Index].stage), Checkpoints[Index].stage,
-							*Breadcrumb->GetFullPath());
-					}
-					GLog->Panic();
+					check(Checkpoints[Index].sType == VK_STRUCTURE_TYPE_CHECKPOINT_DATA_NV);
+					FRHIBreadcrumbNode* Breadcrumb = (FRHIBreadcrumbNode*)Checkpoints[Index].pCheckpointMarker;
+					UE_LOG(LogVulkanRHI, Error, TEXT("[VK_NV_device_diagnostic_checkpoints] %i: Stage %s (0x%08x), %s"),
+						Index, VK_TYPE_TO_STRING(VkPipelineStageFlagBits, Checkpoints[Index].stage), Checkpoints[Index].stage,
+						*Breadcrumb->GetFullPath());
 				}
-			});
+				GLog->Panic();
+			}
+		});
 	}
 #endif // NV_AFTERMATH
 
@@ -635,46 +635,36 @@ void FVulkanDynamicRHI::DumpCrashMarkers()
 		FRHIBreadcrumbState BreadcrumbState;
 
 		Device->ForEachQueue([&QueueRanges, &BreadcrumbState, DeviceIndex](FVulkanQueue& Queue)
+		{
+			FVulkanDiagnosticBuffer* DiagnosticBuffer = Queue.GetDiagnosticBuffer();
+			if (DiagnosticBuffer)
 			{
-				FVulkanDiagnosticBuffer* DiagnosticBuffer = Queue.GetDiagnosticBuffer();
-				if (DiagnosticBuffer)
-				{
-					const EVulkanQueueType QueueType = Queue.GetQueueType();
-					const TCHAR* QueueName = GetVulkanQueueTypeName(QueueType);
-
-					ERHIPipeline Pipeline;
-					switch (QueueType)
-					{
-					default: return; // Skip pipelines that the RHI doesn't handle
-					case EVulkanQueueType::Graphics:     Pipeline = ERHIPipeline::Graphics; break;
-					case EVulkanQueueType::AsyncCompute: Pipeline = ERHIPipeline::AsyncCompute; break;
-					}
+				const EVulkanQueueType QueueType = Queue.GetQueueType();
+				const TCHAR* QueueName = GetVulkanQueueTypeName(QueueType);
 
-					TArray<FRHIBreadcrumbRange>& Ranges = QueueRanges.Add({ DeviceIndex, Pipeline });
+				ERHIPipeline Pipeline;
+				switch (QueueType)
+				{
+				default: return; // Skip pipelines that the RHI doesn't handle
+				case EVulkanQueueType::Graphics:     Pipeline = ERHIPipeline::Graphics; break;
+				case EVulkanQueueType::AsyncCompute: Pipeline = ERHIPipeline::AsyncCompute; break;
+				}
 
-					// Pull all incomplete payloads from the pending interrupt queue.
-					TArray<FVulkanPayload*> Payloads;
-					{
-						FVulkanPayload* Payload;
-						while (Queue.PendingInterrupt.Dequeue(Payload))
-						{
-							Payloads.Add(Payload);
-						}
-					}
+				TArray<FRHIBreadcrumbRange>& Ranges = QueueRanges.Add({ DeviceIndex, Pipeline });
 
-					// Extract the breadcrumb ranges for these payloads.
-					for (FVulkanPayload* Payload : Payloads)
+				// Extract the breadcrumb ranges for the payloads still in the interrupt queue.
+				for (FVulkanPayload* Payload : Queue.PendingInterrupt)
+				{
+					if (Payload->BreadcrumbRange)
 					{
-						if (Payload->BreadcrumbRange)
-						{
-							Ranges.AddUnique(Payload->BreadcrumbRange);
-						}
+						Ranges.AddUnique(Payload->BreadcrumbRange);
 					}
-
-					BreadcrumbState.Devices[DeviceIndex].Pipelines[Pipeline].MarkerOut = DiagnosticBuffer->ReadMarkerOut();
-					BreadcrumbState.Devices[DeviceIndex].Pipelines[Pipeline].MarkerIn = DiagnosticBuffer->ReadMarkerIn();
 				}
-			});
+
+				BreadcrumbState.Devices[DeviceIndex].Pipelines[Pipeline].MarkerOut = DiagnosticBuffer->ReadMarkerOut();
+				BreadcrumbState.Devices[DeviceIndex].Pipelines[Pipeline].MarkerIn = DiagnosticBuffer->ReadMarkerIn();
+			}
+		});
 
 		// Traverse the breadcrumb tree and log active GPU work
 		if (!QueueRanges.IsEmpty())
