Vulkan: Disable anisotropic filtering for 3D textures on Android, as it produces undefined results on some GPUs

diff --git a/Engine/Source/Runtime/Engine/Private/Rendering/Texture3DResource.cpp b/Engine/Source/Runtime/Engine/Private/Rendering/Texture3DResource.cpp
--- a/Engine/Source/Runtime/Engine/Private/Rendering/Texture3DResource.cpp
+++ b/Engine/Source/Runtime/Engine/Private/Rendering/Texture3DResource.cpp
@@ -67,6 +67,11 @@ FTexture3DResource::FTexture3DResource(UVolumeTexture* InOwner, const FStreamabl
 : FStreamableTextureResource(InOwner, InOwner->GetPlatformData(), InState, false)
 , InitialData(InState.RequestedFirstLODIdx())
 {
+	if (!GRHIGlobals.SupportsTexture3DAnisotropicFiltering)
+	{
+		MaxAniso = 1;
+	}
+
 	const int32 FirstLODIdx = InState.RequestedFirstLODIdx();
 
 	// changed to use TryLoadMipsWithSizes
@@ -115,6 +120,10 @@ FTexture3DResource::FTexture3DResource(UVolumeTexture* InOwner, const FTexture3D
 	, ProxiedResource(InProxiedResource)
 	, InitialData(0)
 {
+	if (!GRHIGlobals.SupportsTexture3DAnisotropicFiltering)
+	{
+		MaxAniso = 1;
+	}
 }
 
 void FTexture3DResource::CreateTexture(FRHICommandListBase& RHICmdList)
diff --git a/Engine/Source/Runtime/RHI/Public/RHIGlobals.h b/Engine/Source/Runtime/RHI/Public/RHIGlobals.h
--- a/Engine/Source/Runtime/RHI/Public/RHIGlobals.h
+++ b/Engine/Source/Runtime/RHI/Public/RHIGlobals.h
@@ -227,6 +227,9 @@ struct FRHIGlobals
 	/** Whether to initialize 3D textures using a bulk data (or through a mip update if false). */
 	bool SupportsTexture3D = true;
 
+	/** Whether RHI supports anisotropic filtering for 3D textures. */
+	bool SupportsTexture3DAnisotropicFiltering = true;
+
 	/** true if bulk data should be used with 3d textures */
 	bool UseTexture3DBulkData = false;
 
diff --git a/Engine/Source/Runtime/VulkanRHI/Private/Android/VulkanAndroidPlatform.cpp b/Engine/Source/Runtime/VulkanRHI/Private/Android/VulkanAndroidPlatform.cpp
--- a/Engine/Source/Runtime/VulkanRHI/Private/Android/VulkanAndroidPlatform.cpp
+++ b/Engine/Source/Runtime/VulkanRHI/Private/Android/VulkanAndroidPlatform.cpp
@@ -2108,6 +2108,10 @@ void FVulkanAndroidPlatform::PostInitGPU(const FVulkanDevice& InDevice)
 	SetupRequiresDepthStencilFullWriteWorkaround(InDevice);
 	SetupRequiresMergedDepthStencilLayoutWorkaround(InDevice);
 
+	// The Vulkan specification does not provide concrete guidance on anisotropic filtering for 3D textures, 
+	// and on some Android GPUs this results in undefined behavior
+	GRHIGlobals.SupportsTexture3DAnisotropicFiltering = false;
+
 	// start external compilers if precaching is specified.	
 	static const auto CVarChunkedPSOCache = IConsoleManager::Get().FindConsoleVariable(TEXT("r.Vulkan.UseChunkedPSOCache"));
 	static const auto CVarNumRemoteProgramCompileServices = IConsoleManager::Get().FindConsoleVariable(TEXT("Android.Vulkan.NumRemoteProgramCompileServices"));
