VulkanRHI: Add vk_ext_present_timing functions

#rnx
#Synced-CL 49592606

diff --git a/Engine/Source/Runtime/VulkanRHI/Private/VulkanCommandWrappers.h b/Engine/Source/Runtime/VulkanRHI/Private/VulkanCommandWrappers.h
--- a/Engine/Source/Runtime/VulkanRHI/Private/VulkanCommandWrappers.h
+++ b/Engine/Source/Runtime/VulkanRHI/Private/VulkanCommandWrappers.h
@@ -312,6 +312,11 @@ struct FWrapLayer
 
 	static void CopyMemoryToImageEXT(VkResult Result, VkDevice Device, const VkCopyMemoryToImageInfo* CopyMemoryToImageInfo) VULKAN_LAYER_BODY
 	static void TransitionImageLayoutEXT(VkResult Result, VkDevice Device, uint32_t TransitionCount, const VkHostImageLayoutTransitionInfo* Transitions) VULKAN_LAYER_BODY
+
+	static void GetPastPresentationTimingEXT(VkResult Result, VkDevice Device, const VkPastPresentationTimingInfoEXT* PastPresentationTimingInfo, VkPastPresentationTimingPropertiesEXT* PastPresentationTimingProperties) VULKAN_LAYER_BODY
+	static void GetSwapchainTimeDomainPropertiesEXT(VkResult Result, VkDevice Device, VkSwapchainKHR Swapchain, VkSwapchainTimeDomainPropertiesEXT* SwapchainTimeDomainProperties, uint64_t* TimeDomainsCounter) VULKAN_LAYER_BODY
+	static void GetSwapchainTimingPropertiesEXT(VkResult Result, VkDevice Device, VkSwapchainKHR Swapchain, VkSwapchainTimingPropertiesEXT* SwapchainTimingProperties, uint64_t* SwapchainTimingPropertiesCounter) VULKAN_LAYER_BODY
+	static void SetSwapchainPresentTimingQueueSizeEXT(VkResult Result, VkDevice Device, VkSwapchainKHR Swapchain, uint32_t Size) VULKAN_LAYER_BODY
 };
 
 
@@ -2211,6 +2216,38 @@ namespace VulkanRHI
 		return Result;
 	}
 
+	static FORCEINLINE_DEBUGGABLE VkResult vkGetPastPresentationTimingEXT(VkDevice Device, const VkPastPresentationTimingInfoEXT* PastPresentationTimingInfo, VkPastPresentationTimingPropertiesEXT* PastPresentationTimingProperties)
+	{
+		FWrapLayer::GetPastPresentationTimingEXT(VK_RESULT_MAX_ENUM, Device, PastPresentationTimingInfo, PastPresentationTimingProperties);
+		const VkResult Result = VULKANAPINAMESPACE::vkGetPastPresentationTimingEXT(Device, PastPresentationTimingInfo, PastPresentationTimingProperties);
+		FWrapLayer::GetPastPresentationTimingEXT(Result, Device, PastPresentationTimingInfo, PastPresentationTimingProperties);
+		return Result;
+	}
+
+	static FORCEINLINE_DEBUGGABLE VkResult vkGetSwapchainTimeDomainPropertiesEXT(VkDevice Device, VkSwapchainKHR Swapchain, VkSwapchainTimeDomainPropertiesEXT* SwapchainTimeDomainProperties, uint64_t* TimeDomainsCounter)
+	{
+		FWrapLayer::GetSwapchainTimeDomainPropertiesEXT(VK_RESULT_MAX_ENUM, Device, Swapchain, SwapchainTimeDomainProperties, TimeDomainsCounter);
+		const VkResult Result = VULKANAPINAMESPACE::vkGetSwapchainTimeDomainPropertiesEXT(Device, Swapchain, SwapchainTimeDomainProperties, TimeDomainsCounter);
+		FWrapLayer::GetSwapchainTimeDomainPropertiesEXT(Result, Device, Swapchain, SwapchainTimeDomainProperties, TimeDomainsCounter);
+		return Result;
+	}
+
+	static FORCEINLINE_DEBUGGABLE VkResult vkGetSwapchainTimingPropertiesEXT(VkDevice Device, VkSwapchainKHR Swapchain, VkSwapchainTimingPropertiesEXT* SwapchainTimingProperties, uint64_t* SwapchainTimingPropertiesCounter)
+	{
+		FWrapLayer::GetSwapchainTimingPropertiesEXT(VK_RESULT_MAX_ENUM, Device, Swapchain, SwapchainTimingProperties, SwapchainTimingPropertiesCounter);
+		const VkResult Result = VULKANAPINAMESPACE::vkGetSwapchainTimingPropertiesEXT(Device, Swapchain, SwapchainTimingProperties, SwapchainTimingPropertiesCounter);
+		FWrapLayer::GetSwapchainTimingPropertiesEXT(Result, Device, Swapchain, SwapchainTimingProperties, SwapchainTimingPropertiesCounter);
+		return Result;
+	}
+
+	static FORCEINLINE_DEBUGGABLE VkResult vkSetSwapchainPresentTimingQueueSizeEXT(VkDevice Device, VkSwapchainKHR Swapchain, uint32_t Size)
+	{
+		FWrapLayer::SetSwapchainPresentTimingQueueSizeEXT(VK_RESULT_MAX_ENUM, Device, Swapchain, Size);
+		const VkResult Result = VULKANAPINAMESPACE::vkSetSwapchainPresentTimingQueueSizeEXT(Device, Swapchain, Size);
+		FWrapLayer::SetSwapchainPresentTimingQueueSizeEXT(Result, Device, Swapchain, Size);
+		return Result;
+	}
+
 
 #if VULKAN_ENABLE_IMAGE_TRACKING_LAYER
 	void BindDebugLabelName(VkImage Image, const TCHAR* Name);
diff --git a/Engine/Source/Runtime/VulkanRHI/Private/VulkanDebug.cpp b/Engine/Source/Runtime/VulkanRHI/Private/VulkanDebug.cpp
--- a/Engine/Source/Runtime/VulkanRHI/Private/VulkanDebug.cpp
+++ b/Engine/Source/Runtime/VulkanRHI/Private/VulkanDebug.cpp
@@ -5158,6 +5158,45 @@ void FWrapLayer::TransitionImageLayoutEXT(VkResult Result, VkDevice Device, uint
 	}
 }
 
+void FWrapLayer::GetPastPresentationTimingEXT(VkResult Result, VkDevice Device, const VkPastPresentationTimingInfoEXT* PastPresentationTimingInfo, VkPastPresentationTimingPropertiesEXT* PastPresentationTimingProperties)
+{
+	if (Result == VK_RESULT_MAX_ENUM)
+	{
+#if VULKAN_ENABLE_DUMP_LAYER
+		PrintfBeginResult(FString::Printf(TEXT("GetPastPresentationTimingEXT()")));
+#endif
+	}
+}
+
+void FWrapLayer::GetSwapchainTimeDomainPropertiesEXT(VkResult Result, VkDevice Device, VkSwapchainKHR Swapchain, VkSwapchainTimeDomainPropertiesEXT* SwapchainTimeDomainProperties, uint64_t* TimeDomainsCounter)
+{
+	if (Result == VK_RESULT_MAX_ENUM)
+	{
+#if VULKAN_ENABLE_DUMP_LAYER
+		PrintfBeginResult(FString::Printf(TEXT("GetSwapchainTimeDomainPropertiesEXT()")));
+#endif
+	}
+}
+
+void FWrapLayer::GetSwapchainTimingPropertiesEXT(VkResult Result, VkDevice Device, VkSwapchainKHR Swapchain, VkSwapchainTimingPropertiesEXT* SwapchainTimingProperties, uint64_t* SwapchainTimingPropertiesCounter)
+{
+	if (Result == VK_RESULT_MAX_ENUM)
+	{
+#if VULKAN_ENABLE_DUMP_LAYER
+		PrintfBeginResult(FString::Printf(TEXT("GetSwapchainTimingPropertiesEXT()")));
+#endif
+	}
+}
+
+void FWrapLayer::SetSwapchainPresentTimingQueueSizeEXT(VkResult Result, VkDevice Device, VkSwapchainKHR Swapchain, uint32_t Size)
+{
+	if (Result == VK_RESULT_MAX_ENUM)
+	{
+#if VULKAN_ENABLE_DUMP_LAYER
+		PrintfBeginResult(FString::Printf(TEXT("SetSwapchainPresentTimingQueueSizeEXT()")));
+#endif
+	}
+}
 
 
 #if VULKAN_ENABLE_IMAGE_TRACKING_LAYER
diff --git a/Engine/Source/Runtime/VulkanRHI/Private/VulkanLoader.h b/Engine/Source/Runtime/VulkanRHI/Private/VulkanLoader.h
--- a/Engine/Source/Runtime/VulkanRHI/Private/VulkanLoader.h
+++ b/Engine/Source/Runtime/VulkanRHI/Private/VulkanLoader.h
@@ -286,7 +286,11 @@
 	EnumMacro(PFN_vkDestroyIndirectExecutionSetEXT, vkDestroyIndirectExecutionSetEXT) \
 	EnumMacro(PFN_vkGetGeneratedCommandsMemoryRequirementsEXT, vkGetGeneratedCommandsMemoryRequirementsEXT) \
 	EnumMacro(PFN_vkUpdateIndirectExecutionSetPipelineEXT, vkUpdateIndirectExecutionSetPipelineEXT) \
-	EnumMacro(PFN_vkUpdateIndirectExecutionSetShaderEXT, vkUpdateIndirectExecutionSetShaderEXT)
+	EnumMacro(PFN_vkUpdateIndirectExecutionSetShaderEXT, vkUpdateIndirectExecutionSetShaderEXT) \
+	EnumMacro(PFN_vkGetPastPresentationTimingEXT, vkGetPastPresentationTimingEXT) \
+	EnumMacro(PFN_vkGetSwapchainTimeDomainPropertiesEXT, vkGetSwapchainTimeDomainPropertiesEXT) \
+	EnumMacro(PFN_vkGetSwapchainTimingPropertiesEXT, vkGetSwapchainTimingPropertiesEXT) \
+	EnumMacro(PFN_vkSetSwapchainPresentTimingQueueSizeEXT, vkSetSwapchainPresentTimingQueueSizeEXT)
 
 
 
