VulkanRHI: Add support for maintenance9. Use VK_QUERY_POOL_CREATE_RESET_BIT_KHR when available to reset queries on creation.  Avoid empty command buffers when EXTHostQueryReset is supported.

#rnx
#Synced-CL 49592606

diff --git a/Engine/Source/Runtime/VulkanRHI/Private/VulkanDevice.h b/Engine/Source/Runtime/VulkanRHI/Private/VulkanDevice.h
--- a/Engine/Source/Runtime/VulkanRHI/Private/VulkanDevice.h
+++ b/Engine/Source/Runtime/VulkanRHI/Private/VulkanDevice.h
@@ -65,6 +65,7 @@ struct FOptionalVulkanDeviceExtensions
 			uint64 HasEXTImageCompressionControl : 1;
 			uint64 HasEXTMutableDescriptorType : 1;
 			uint64 HasKHRMaintenance7 : 1;
+			uint64 HasKHRMaintenance9 : 1;
 			uint64 HasEXTShaderObject : 1;
 			uint64 HasEXTGraphicsPipelineLibrary : 1;
 			uint64 HasEXTDeviceGeneratedCommands : 1;
diff --git a/Engine/Source/Runtime/VulkanRHI/Private/VulkanExtensions.cpp b/Engine/Source/Runtime/VulkanRHI/Private/VulkanExtensions.cpp
--- a/Engine/Source/Runtime/VulkanRHI/Private/VulkanExtensions.cpp
+++ b/Engine/Source/Runtime/VulkanRHI/Private/VulkanExtensions.cpp
@@ -348,6 +348,7 @@ public:
 private:
 	VkPhysicalDeviceMaintenance6FeaturesKHR Maintenance6Features;
 };
+
 // ***** VK_KHR_maintenance7
 class FVulkanKHRMaintenance7Extension : public FVulkanDeviceExtension
 {
@@ -395,6 +396,42 @@ private:
 	VkPhysicalDeviceMaintenance7FeaturesKHR Maintenance7Features;
 };
 
+// ***** VK_KHR_maintenance9
+class FVulkanKHRMaintenance9Extension : public FVulkanDeviceExtension
+{
+public:
+
+	FVulkanKHRMaintenance9Extension(FVulkanDevice* InDevice)
+		: FVulkanDeviceExtension(InDevice, VK_KHR_MAINTENANCE_9_EXTENSION_NAME, VULKAN_EXTENSION_ENABLED)
+	{
+	}
+
+	virtual void PrePhysicalDeviceFeatures(VkPhysicalDeviceFeatures2KHR& PhysicalDeviceFeatures2) override final
+	{
+		ZeroVulkanStruct(Maintenance9Features, VK_STRUCTURE_TYPE_PHYSICAL_DEVICE_MAINTENANCE_9_FEATURES_KHR);
+		AddToPNext(PhysicalDeviceFeatures2, Maintenance9Features);
+	}
+
+	virtual void PostPhysicalDeviceFeatures(FOptionalVulkanDeviceExtensions& ExtensionFlags) override final
+	{
+		bRequirementsPassed = (Maintenance9Features.maintenance9 == VK_TRUE);
+		if (bRequirementsPassed)
+		{
+			ExtensionFlags.HasKHRMaintenance9 = true;
+		}
+	}
+
+	virtual void PreCreateDevice(VkDeviceCreateInfo& DeviceCreateInfo) override final
+	{
+		if (bRequirementsPassed)
+		{
+			AddToPNext(DeviceCreateInfo, Maintenance9Features);
+		}
+	}
+
+private:
+	VkPhysicalDeviceMaintenance9FeaturesKHR Maintenance9Features;
+};
 
 // ***** VK_KHR_driver_properties
 class FVulkanKHRDriverPropertiesExtension : public FVulkanDeviceExtension
@@ -2505,6 +2542,7 @@ FVulkanDeviceExtensionArray FVulkanDeviceExtension::GetUESupportedDeviceExtensio
 	ADD_CUSTOM_EXTENSION(FVulkanKHRMaintenance5Extension);
 	ADD_CUSTOM_EXTENSION(FVulkanKHRMaintenance6Extension);
 	ADD_CUSTOM_EXTENSION(FVulkanKHRMaintenance7Extension);
+	ADD_CUSTOM_EXTENSION(FVulkanKHRMaintenance9Extension);
 	ADD_CUSTOM_EXTENSION(FVulkanShaderAtomicInt64Extension);
 	ADD_CUSTOM_EXTENSION(FVulkanShaderImageAtomicInt64Extension);
 	ADD_CUSTOM_EXTENSION(FVulkanEXTScalarBlockLayoutExtension);
diff --git a/Engine/Source/Runtime/VulkanRHI/Private/VulkanQuery.cpp b/Engine/Source/Runtime/VulkanRHI/Private/VulkanQuery.cpp
--- a/Engine/Source/Runtime/VulkanRHI/Private/VulkanQuery.cpp
+++ b/Engine/Source/Runtime/VulkanRHI/Private/VulkanQuery.cpp
@@ -59,10 +59,24 @@ FVulkanQueryPool::FVulkanQueryPool(FVulkanDevice& InDevice, uint32 InMaxQueries,
 	, QueryType(InQueryType)
 {
 	INC_DWORD_STAT(STAT_VulkanNumQueryPools);
-	VkQueryPoolCreateInfo PoolCreateInfo;
-	ZeroVulkanStruct(PoolCreateInfo, VK_STRUCTURE_TYPE_QUERY_POOL_CREATE_INFO);
-	PoolCreateInfo.queryType = GetVkQueryType(InQueryType);
-	PoolCreateInfo.queryCount = MaxQueries;
+	VkQueryPoolCreateInfo PoolCreateInfo = {
+		.sType = VK_STRUCTURE_TYPE_QUERY_POOL_CREATE_INFO,
+		.pNext = nullptr,
+		.flags = 0,
+		.queryType = GetVkQueryType(InQueryType),
+		.queryCount = MaxQueries,
+		.pipelineStatistics = 0
+	};
+
+	// With maintenance9, always set an initial state for the queries
+	if (Device.GetOptionalExtensions().HasKHRMaintenance9)
+	{
+		QueryRefs.AddDefaulted(MaxQueries / QueryStride);
+		QueryResults.SetNumZeroed(MaxQueries / QueryStride);
+		PoolCreateInfo.flags |= VK_QUERY_POOL_CREATE_RESET_BIT_KHR;
+		bIsDirty = false;
+	}
+
 	VERIFYVULKANRESULT(VulkanRHI::vkCreateQueryPool(Device.GetHandle(), &PoolCreateInfo, VULKAN_CPU_ALLOCATOR, &QueryPool));
 }
 
@@ -80,6 +94,7 @@ void FVulkanQueryPool::ReserveQuery(FVulkanRenderQuery* Query)
 	Query->IndexInPool = QuerySlot * QueryStride;
 	QueryRefs[QuerySlot] = Query;
 	QueryResults[QuerySlot] = &Query->Result;
+	bIsDirty = true;
 }
 
 uint32 FVulkanQueryPool::ReserveQuery(uint64* ResultPtr)
@@ -89,10 +104,11 @@ uint32 FVulkanQueryPool::ReserveQuery(uint64* ResultPtr)
 	const uint32 IndexInPool = QuerySlot * QueryStride;
 	QueryRefs[QuerySlot] = nullptr;
 	QueryResults[QuerySlot] = ResultPtr;
+	bIsDirty = true;
 	return IndexInPool;
 }
 
-void FVulkanQueryPool::Reset(FVulkanCommandBuffer& InCmdBuffer, uint32 InQueryStride)
+void FVulkanQueryPool::Reset(FVulkanCommandBuffer* InCmdBuffer, uint32 InQueryStride)
 {
 	CurrentQueryCount = 0;
 	QueryStride = InQueryStride;
@@ -103,14 +119,17 @@ void FVulkanQueryPool::Reset(FVulkanCommandBuffer& InCmdBuffer, uint32 InQuerySt
 	QueryRefs.AddDefaulted(MaxQueries / QueryStride);
 	QueryResults.SetNumZeroed(MaxQueries / QueryStride);
 
-	if (Device.GetOptionalExtensions().HasEXTHostQueryReset)
+	if (InCmdBuffer)
 	{
-		VulkanRHI::vkResetQueryPoolEXT(GetDevice().GetHandle(), QueryPool, 0, MaxQueries);
+		VulkanRHI::vkCmdResetQueryPool(InCmdBuffer->GetHandle(), QueryPool, 0, MaxQueries);
 	}
 	else
 	{
-		VulkanRHI::vkCmdResetQueryPool(InCmdBuffer.GetHandle(), QueryPool, 0, MaxQueries);
+		checkf(Device.GetOptionalExtensions().HasEXTHostQueryReset, TEXT("VK_EXT_host_query_reset must be supported to reset queries without a command buffer."));
+		VulkanRHI::vkResetQueryPoolEXT(Device.GetHandle(), QueryPool, 0, MaxQueries);
 	}
+
+	bIsDirty = false;
 }
 
 bool FVulkanQueryPool::IsStale() const
@@ -120,10 +139,14 @@ bool FVulkanQueryPool::IsStale() const
 
 void FVulkanCommandListContext::BeginOcclusionQueryBatch(uint32 NumQueriesInBatch, uint32 MultiViewCount)
 {
-	uint32 QueryStride = MultiViewCount == 0 ? 1 : MultiViewCount;
+	const uint32 QueryStride = MultiViewCount == 0 ? 1 : MultiViewCount;
 	TArray<FVulkanQueryPool*>& OcclusionPoolArray = GetQueryPoolArray(EVulkanQueryPoolType::Occlusion);
 	FVulkanQueryPool* NewOcclusionQueryPool = Device.AcquireOcclusionQueryPool(NumQueriesInBatch * QueryStride);
-	NewOcclusionQueryPool->Reset(GetCommandBuffer(), QueryStride);
+	if (NewOcclusionQueryPool->IsDirty() || (NewOcclusionQueryPool->GetQueryStride() != QueryStride))
+	{
+		FVulkanCommandBuffer* CommandBuffer = Device.GetOptionalExtensions().HasEXTHostQueryReset ? nullptr : &GetCommandBuffer();
+		NewOcclusionQueryPool->Reset(CommandBuffer, QueryStride);
+	}
 	OcclusionPoolArray.Add(NewOcclusionQueryPool);
 }
 
@@ -177,6 +200,7 @@ FVulkanQueryPool* FVulkanDevice::AcquireTimingQueryPool()
 void FVulkanDevice::ReleaseQueryPool(FVulkanQueryPool* Pool)
 {
 	FScopeLock Lock(&QueryPoolLock);
+	Pool->SetDirty();
 	FreeQueryPools[(int32)Pool->GetPoolType()].Add(Pool);
 }
 
@@ -387,12 +411,19 @@ FVulkanQueryPool* FVulkanContextCommon::GetCurrentTimestampQueryPool(FVulkanPayl
 	if ((TimestampPoolArray.Num() == 0) || TimestampPoolArray.Last()->IsFull() || TimestampPoolArray.Last()->GetQueryStride() != DesiredQueryStride)
 	{
 		FVulkanQueryPool* NewPool = Device.AcquireTimingQueryPool();
-
-		if (Payload.CommandBuffers.Num() == 0)
+		if (NewPool->IsDirty() || (NewPool->GetQueryStride() != DesiredQueryStride))
 		{
-			PrepareNewCommandBuffer(Payload);
+			FVulkanCommandBuffer* CommandBuffer = nullptr;
+			if (!Device.GetOptionalExtensions().HasEXTHostQueryReset)
+			{
+				if (Payload.CommandBuffers.Num() == 0)
+				{
+					PrepareNewCommandBuffer(Payload);
+				}
+				CommandBuffer = Payload.CommandBuffers.Last();
+			}
+			NewPool->Reset(CommandBuffer, DesiredQueryStride);
 		}
-		NewPool->Reset(*Payload.CommandBuffers.Last(), DesiredQueryStride);
 
 		TimestampPoolArray.Add(NewPool);
 	}
diff --git a/Engine/Source/Runtime/VulkanRHI/Private/VulkanQuery.h b/Engine/Source/Runtime/VulkanRHI/Private/VulkanQuery.h
--- a/Engine/Source/Runtime/VulkanRHI/Private/VulkanQuery.h
+++ b/Engine/Source/Runtime/VulkanRHI/Private/VulkanQuery.h
@@ -84,7 +84,7 @@ public:
 	void ReserveQuery(FVulkanRenderQuery* Query);
 	uint32 ReserveQuery(uint64* ResultPtr);
 
-	void Reset(FVulkanCommandBuffer& InCmdBuffer, uint32 InQueryStride = 1u);
+	void Reset(FVulkanCommandBuffer* InCmdBuffer, uint32 InQueryStride = 1u);
 	bool IsStale() const;
 
 	void IncrementUnusedFrameCount()
@@ -92,6 +92,16 @@ public:
 		UnusedFrameCount++;
 	}
 
+	void SetDirty()
+	{
+		bIsDirty = true;
+	}
+
+	bool IsDirty() const
+	{
+		return bIsDirty;
+	}
+
 protected:
 	FVulkanDevice& Device;
 	VkQueryPool QueryPool;
@@ -106,6 +116,7 @@ protected:
 	 *  can be summed together to get the actual value. */
 	uint32 QueryStride = 1u;
 	int32 UnusedFrameCount = 0;
+	bool bIsDirty = true;
 
 	friend FVulkanDynamicRHI;
 };
diff --git a/Engine/Source/Runtime/VulkanRHI/Private/VulkanRayTracing.cpp b/Engine/Source/Runtime/VulkanRHI/Private/VulkanRayTracing.cpp
--- a/Engine/Source/Runtime/VulkanRHI/Private/VulkanRayTracing.cpp
+++ b/Engine/Source/Runtime/VulkanRHI/Private/VulkanRayTracing.cpp
@@ -1734,7 +1734,10 @@ FVulkanRayTracingCompactedSizeQueryPool::FVulkanRayTracingCompactedSizeQueryPool
 	: FVulkanQueryPool(InDevice, InMaxQueries, EVulkanQueryPoolType::ASCompactedSize)
 {
 	QueryOutput.SetNumZeroed(InMaxQueries);
-	VulkanRHI::vkResetQueryPoolEXT(InDevice.GetHandle(), QueryPool, 0, MaxQueries);
+	if (bIsDirty)
+	{
+		VulkanRHI::vkResetQueryPoolEXT(InDevice.GetHandle(), QueryPool, 0, MaxQueries);
+	}
 }
 
 void FVulkanRayTracingCompactedSizeQueryPool::EndBatch(FVulkanCommandListContext& CommandContext)
@@ -1742,9 +1745,9 @@ void FVulkanRayTracingCompactedSizeQueryPool::EndBatch(FVulkanCommandListContext
 	SyncPoint = CommandContext.GetContextSyncPoint();
 }
 
-void FVulkanRayTracingCompactedSizeQueryPool::Reset(FVulkanCommandBuffer& InCmdBuffer)
+void FVulkanRayTracingCompactedSizeQueryPool::Reset()
 {
-	VulkanRHI::vkCmdResetQueryPool(InCmdBuffer.GetHandle(), QueryPool, 0, MaxQueries);
+	VulkanRHI::vkResetQueryPoolEXT(Device.GetHandle(), QueryPool, 0, MaxQueries);
 	SyncPoint = nullptr;
 	check(QueryOutput.Num() == MaxQueries);
 	FMemory::Memzero(QueryOutput.GetData(), MaxQueries * sizeof(uint64));
@@ -1825,10 +1828,10 @@ void FVulkanRayTracingCompactionRequestHandler::Update(FVulkanCommandListContext
 	// If we have an active batch, wait on those queries and launch compaction when the complete
 	if (ActiveBLASes.Num() > 0)
 	{		
-		FVulkanCommandBuffer& CmdBuffer = InCommandContext.GetCommandBuffer();
-
 		if (QueryPool->TryGetResults(ActiveBLASes.Num()))
 		{
+			FVulkanCommandBuffer& CmdBuffer = InCommandContext.GetCommandBuffer();
+
 			// Compact
 			for (int32 BLASIndex = 0; BLASIndex < ActiveBLASes.Num(); ++BLASIndex)
 			{
@@ -1838,7 +1841,7 @@ void FVulkanRayTracingCompactionRequestHandler::Update(FVulkanCommandListContext
 				}
 			}
 
-			QueryPool->Reset(CmdBuffer);
+			QueryPool->Reset();
 
 			ActiveBLASes.Empty(GVulkanRayTracingMaxBatchedCompaction);
 
diff --git a/Engine/Source/Runtime/VulkanRHI/Private/VulkanRayTracing.h b/Engine/Source/Runtime/VulkanRHI/Private/VulkanRayTracing.h
--- a/Engine/Source/Runtime/VulkanRHI/Private/VulkanRayTracing.h
+++ b/Engine/Source/Runtime/VulkanRHI/Private/VulkanRayTracing.h
@@ -327,7 +327,7 @@ public:
 
 	void EndBatch(FVulkanCommandListContext& CommandContext);
 	bool TryGetResults(uint32 NumResults);
-	void Reset(FVulkanCommandBuffer& InCmdBuffer);
+	void Reset();
 
 	TArray<uint64> QueryOutput;
 	FVulkanSyncPointRef SyncPoint;
