Vulkan Bindless
- Changing Vulkan bindless descriptor management to use ERHIDescriptorType over VkDescriptorType.
- Changing descriptor update names to all match.
- Adding bulk descriptor allocations.

#Synced-CL 49664761

diff --git a/Engine/Source/Runtime/VulkanRHI/Private/VulkanBindlessDescriptorManager.cpp b/Engine/Source/Runtime/VulkanRHI/Private/VulkanBindlessDescriptorManager.cpp
--- a/Engine/Source/Runtime/VulkanRHI/Private/VulkanBindlessDescriptorManager.cpp
+++ b/Engine/Source/Runtime/VulkanRHI/Private/VulkanBindlessDescriptorManager.cpp
@@ -157,6 +157,24 @@ static constexpr ERHIDescriptorType GetUEDescriptorTypeFromVkDescriptorType(VkDe
 	return ERHIDescriptorType::Invalid;
 }
 
+static constexpr VkDescriptorType GetVkDescriptorTypeFromUEDescriptorType(ERHIDescriptorType DescriptorType)
+{
+	switch (DescriptorType)
+	{
+	case ERHIDescriptorType::Sampler:               return VK_DESCRIPTOR_TYPE_SAMPLER;
+	case ERHIDescriptorType::TextureSRV:            return VK_DESCRIPTOR_TYPE_SAMPLED_IMAGE;
+	case ERHIDescriptorType::TextureUAV:            return VK_DESCRIPTOR_TYPE_STORAGE_IMAGE;
+	case ERHIDescriptorType::TypedBufferSRV:        return VK_DESCRIPTOR_TYPE_UNIFORM_TEXEL_BUFFER;
+	case ERHIDescriptorType::TypedBufferUAV:        return VK_DESCRIPTOR_TYPE_STORAGE_TEXEL_BUFFER;
+	case ERHIDescriptorType::BufferUAV:             return VK_DESCRIPTOR_TYPE_STORAGE_BUFFER;
+	case ERHIDescriptorType::BufferSRV:             return VK_DESCRIPTOR_TYPE_UNIFORM_BUFFER;
+	case ERHIDescriptorType::AccelerationStructure: return VK_DESCRIPTOR_TYPE_ACCELERATION_STRUCTURE_KHR;
+	default: checkNoEntry();
+	}
+
+	return VK_DESCRIPTOR_TYPE_MAX_ENUM;
+}
+
 extern TAutoConsoleVariable<int32> GCVarRobustBufferAccess;
 static inline uint32 GetDescriptorTypeSize(FVulkanDevice& Device, VkDescriptorType DescriptorType)
 {
@@ -248,10 +266,8 @@ inline FVulkanBindlessDescriptorManager::BindlessSetState& FVulkanBindlessDescri
 }
 
 // Check all the requirements to be running in Bindless using Descriptor Buffers
-bool FVulkanBindlessDescriptorManager::VerifySupport(FVulkanDevice& Device)
+static bool VerifySupport(FVulkanDevice& Device, ERHIBindlessConfiguration BindlessConfig)
 {
-	const ERHIBindlessConfiguration BindlessConfig = RHIGetRuntimeBindlessConfiguration(GMaxRHIShaderPlatform);
-
 	if (BindlessConfig == ERHIBindlessConfiguration::Disabled)
 	{
 		return false;
@@ -315,7 +331,6 @@ bool FVulkanBindlessDescriptorManager::VerifySupport(FVulkanDevice& Device)
 
 FVulkanBindlessDescriptorManager::FVulkanBindlessDescriptorManager(FVulkanDevice& InDevice)
 	: Device(InDevice)
-	, bIsSupported(VerifySupport(InDevice))
 {
 	FMemory::Memzero(BufferBindingInfo);
 	for (uint32 Index = 0; Index < VulkanBindless::MaxNumSets; Index++)
@@ -329,49 +344,13 @@ FVulkanBindlessDescriptorManager::~FVulkanBindlessDescriptorManager()
 	checkf(BindlessPipelineLayout == VK_NULL_HANDLE, TEXT("DeInit() was not called on FVulkanBindlessDescriptorManager!"));
 }
 
-void FVulkanBindlessDescriptorManager::Deinit()
-{
-	const VkDevice DeviceHandle = Device.GetHandle();
-
-	if (bIsSupported)
-	{
-		VulkanRHI::vkDestroyPipelineLayout(DeviceHandle, BindlessPipelineLayout, VULKAN_CPU_ALLOCATOR);
-		BindlessPipelineLayout = VK_NULL_HANDLE;
-
-		auto DestroyBindlessState = [DeviceHandle](BindlessSetState& State)
-			{
-				VulkanRHI::vkDestroyDescriptorSetLayout(DeviceHandle, State.DescriptorSetLayout, VULKAN_CPU_ALLOCATOR);
-				State.DescriptorSetLayout = VK_NULL_HANDLE;
-
-				VulkanRHI::vkDestroyBuffer(DeviceHandle, State.BufferHandle, VULKAN_CPU_ALLOCATOR);
-				State.BufferHandle = VK_NULL_HANDLE;
-
-				VulkanRHI::vkUnmapMemory(DeviceHandle, State.MemoryHandle);
-				VulkanRHI::vkFreeMemory(DeviceHandle, State.MemoryHandle, VULKAN_CPU_ALLOCATOR);
-				State.MemoryHandle = VK_NULL_HANDLE;
-			};
-
-		for (uint32 SetIndex = 0; SetIndex < VulkanBindless::NumBindlessSets; ++SetIndex)
-		{
-			BindlessSetState& State = BindlessSetStates[SetIndex];
-			if (!State.DescriptorTypes.IsEmpty())
-			{
-				DestroyBindlessState(State);
-			}
-		}
-
-		VulkanRHI::vkDestroyDescriptorSetLayout(DeviceHandle, SingleUseUBDescriptorSetLayout, VULKAN_CPU_ALLOCATOR);
-		SingleUseUBDescriptorSetLayout = VK_NULL_HANDLE;
-
-		VulkanRHI::vkDestroyDescriptorSetLayout(DeviceHandle, EmptyDescriptorSetLayout, VULKAN_CPU_ALLOCATOR);
-		EmptyDescriptorSetLayout = VK_NULL_HANDLE;
-	}
-}
-
 void FVulkanBindlessDescriptorManager::Init()
 {
-	if (!bIsSupported)
+	Configuration = RHIGetRuntimeBindlessConfiguration(GMaxRHIShaderPlatform);
+
+	if (!VerifySupport(Device, Configuration))
 	{
+		Configuration = ERHIBindlessConfiguration::Disabled;
 		return;
 	}
 
@@ -655,9 +634,48 @@ void FVulkanBindlessDescriptorManager::Init()
 	}
 }
 
+void FVulkanBindlessDescriptorManager::Deinit()
+{
+	const VkDevice DeviceHandle = Device.GetHandle();
+
+	if (IsSupported())
+	{
+		VulkanRHI::vkDestroyPipelineLayout(DeviceHandle, BindlessPipelineLayout, VULKAN_CPU_ALLOCATOR);
+		BindlessPipelineLayout = VK_NULL_HANDLE;
+
+		auto DestroyBindlessState = [DeviceHandle](BindlessSetState& State)
+			{
+				VulkanRHI::vkDestroyDescriptorSetLayout(DeviceHandle, State.DescriptorSetLayout, VULKAN_CPU_ALLOCATOR);
+				State.DescriptorSetLayout = VK_NULL_HANDLE;
+
+				VulkanRHI::vkDestroyBuffer(DeviceHandle, State.BufferHandle, VULKAN_CPU_ALLOCATOR);
+				State.BufferHandle = VK_NULL_HANDLE;
+
+				VulkanRHI::vkUnmapMemory(DeviceHandle, State.MemoryHandle);
+				VulkanRHI::vkFreeMemory(DeviceHandle, State.MemoryHandle, VULKAN_CPU_ALLOCATOR);
+				State.MemoryHandle = VK_NULL_HANDLE;
+			};
+
+		for (uint32 SetIndex = 0; SetIndex < VulkanBindless::NumBindlessSets; ++SetIndex)
+		{
+			BindlessSetState& State = BindlessSetStates[SetIndex];
+			if (!State.DescriptorTypes.IsEmpty())
+			{
+				DestroyBindlessState(State);
+			}
+		}
+
+		VulkanRHI::vkDestroyDescriptorSetLayout(DeviceHandle, SingleUseUBDescriptorSetLayout, VULKAN_CPU_ALLOCATOR);
+		SingleUseUBDescriptorSetLayout = VK_NULL_HANDLE;
+
+		VulkanRHI::vkDestroyDescriptorSetLayout(DeviceHandle, EmptyDescriptorSetLayout, VULKAN_CPU_ALLOCATOR);
+		EmptyDescriptorSetLayout = VK_NULL_HANDLE;
+	}
+}
+
 void FVulkanBindlessDescriptorManager::BindDescriptorBuffers(VkCommandBuffer CommandBuffer, VkPipelineStageFlags SupportedStages)
 {
-	checkf(bIsSupported, TEXT("Trying to BindDescriptorBuffers but bindless is not supported!"));
+	checkf(IsSupported(), TEXT("Trying to BindDescriptorBuffers but bindless is not supported!"));
 
 	VulkanRHI::vkCmdBindDescriptorBuffersEXT(CommandBuffer, VulkanBindless::NumBindlessSets, BufferBindingInfo);
 
@@ -679,7 +697,7 @@ void FVulkanBindlessDescriptorManager::BindDescriptorBuffers(VkCommandBuffer Com
 
 void FVulkanBindlessDescriptorManager::RegisterUniformBuffers(FVulkanCommandListContext& Context, VkShaderStageFlags StageFlags, TConstArrayView<VkWriteDescriptorSet> UBDescriptors)
 {
-	checkf(bIsSupported, TEXT("Trying to RegisterUniformBuffers but bindless is not supported!"));
+	checkf(IsSupported(), TEXT("Trying to RegisterUniformBuffers but bindless is not supported!"));
 	check(UBDescriptors.Num() < VulkanBindless::MaxUniformBuffersTotal);
 
 	VkPushDescriptorSetInfo PushDescriptorSetInfo =
@@ -696,35 +714,45 @@ void FVulkanBindlessDescriptorManager::RegisterUniformBuffers(FVulkanCommandList
 	VulkanRHI::vkCmdPushDescriptorSet2KHR(Context.GetCommandBuffer().GetHandle(), &PushDescriptorSetInfo);
 }
 
-FRHIDescriptorHandle FVulkanBindlessDescriptorManager::AllocateDescriptor(VkDescriptorType DescriptorType)
+FRHIDescriptorHandle FVulkanBindlessDescriptorManager::AllocateDescriptor(ERHIDescriptorType DescriptorType)
 {
 	if (IsSupported())
 	{
-		ERHIDescriptorType UEType = GetUEDescriptorTypeFromVkDescriptorType(DescriptorType);
-		return GetBindlessState(DescriptorType).Allocator->Allocate(UEType);
+		return GetBindlessState(DescriptorType).Allocator->Allocate(DescriptorType);
 	}
 
 	return FRHIDescriptorHandle();
 }
 
-FRHIDescriptorHandle FVulkanBindlessDescriptorManager::AllocateDescriptor(VkDescriptorType DescriptorType, VkDeviceAddress BufferAddress, VkDeviceSize BufferSize)
+TOptional<FRHIDescriptorAllocation> FVulkanBindlessDescriptorManager::AllocateDescriptors(ERHIDescriptorType DescriptorType, uint32 DescriptorCount)
+{
+	if (IsSupported())
+	{
+		return GetBindlessState(DescriptorType).Allocator->Allocate(DescriptorCount);
+	}
+
+	return TOptional<FRHIDescriptorAllocation>();
+}
+
+
+FRHIDescriptorHandle FVulkanBindlessDescriptorManager::AllocateDescriptor(ERHIDescriptorType DescriptorType, VkDeviceAddress BufferAddress, VkDeviceSize BufferSize)
 {
 	if (IsSupported())
 	{
 		FRHIDescriptorHandle DescriptorHandle = AllocateDescriptor(DescriptorType);
-		UpdateBuffer({}, DescriptorHandle, DescriptorType, BufferAddress, BufferSize, true);
+		UpdateDescriptor({}, DescriptorHandle, BufferAddress, BufferSize, true);
 		return DescriptorHandle;
 	}
 
 	return FRHIDescriptorHandle();
 }
 
-FRHIDescriptorHandle FVulkanBindlessDescriptorManager::AllocateDescriptor(VkDescriptorType DescriptorType, FVulkanBuffer* Buffer, uint32 ExtraOffset)
+FRHIDescriptorHandle FVulkanBindlessDescriptorManager::AllocateDescriptor(ERHIDescriptorType DescriptorType, FVulkanBuffer* Buffer, uint32 ExtraOffset)
 {
 	if (IsSupported())
 	{
 		FRHIDescriptorHandle DescriptorHandle = AllocateDescriptor(DescriptorType);
-		UpdateBuffer({}, DescriptorHandle, DescriptorType, Buffer->GetHandle(), Buffer->GetOffset() + ExtraOffset, Buffer->GetCurrentSize() - ExtraOffset, true);
+		UpdateDescriptor({}, DescriptorHandle, Buffer->GetHandle(), Buffer->GetOffset() + ExtraOffset, Buffer->GetCurrentSize() - ExtraOffset, true);
 		return DescriptorHandle;
 	}
 
@@ -735,8 +763,8 @@ FRHIDescriptorHandle FVulkanBindlessDescriptorManager::AllocateDescriptor(VkSamp
 {
 	if (IsSupported())
 	{
-		FRHIDescriptorHandle DescriptorHandle = AllocateDescriptor(VK_DESCRIPTOR_TYPE_SAMPLER);
-		UpdateSampler({}, DescriptorHandle, VulkanSampler);
+		FRHIDescriptorHandle DescriptorHandle = AllocateDescriptor(ERHIDescriptorType::Sampler);
+		UpdateDescriptor({}, DescriptorHandle, VulkanSampler);
 		return DescriptorHandle;
 	}
 
@@ -755,10 +783,24 @@ void FVulkanBindlessDescriptorManager::FreeDescriptor(FRHIDescriptorHandle Descr
 	}
 }
 
-void FVulkanBindlessDescriptorManager::UpdateDescriptor(const FVulkanContextArray& Contexts, FRHIDescriptorHandle DescriptorHandle, VkDescriptorType DescriptorType, VkDescriptorDataEXT DescriptorData, bool bImmediateUpdate)
+void FVulkanBindlessDescriptorManager::FreeDescriptors(FRHIDescriptorAllocation Allocation, VkDescriptorType DescriptorType)
+{
+	if (Allocation.Count > 0)
+	{
+		BindlessSetState& State = GetBindlessState(DescriptorType);
+		State.Allocator->Free(Allocation);
+
+		const uint32 ByteOffset = Allocation.StartIndex * State.DescriptorSize;
+		FMemory::Memzero(&State.CPUDescriptorMemory[ByteOffset], State.DescriptorSize * Allocation.Count); // easier for debugging for now
+	}
+}
+
+void FVulkanBindlessDescriptorManager::UpdateDescriptorInternal(const FVulkanContextArray& Contexts, FRHIDescriptorHandle DescriptorHandle, VkDescriptorDataEXT DescriptorData, bool bImmediateUpdate)
 {
 	checkf(DescriptorHandle.IsValid(), TEXT("Attemping to update invalid descriptor handle!"));
 
+	const VkDescriptorType DescriptorType = GetVkDescriptorTypeFromUEDescriptorType(DescriptorHandle.GetType());
+
 	BindlessSetState& State = GetBindlessState(DescriptorHandle.GetType());
 	const uint32 ByteOffset = DescriptorHandle.GetIndex() * State.DescriptorSize;
 	checkSlow(State.DescriptorTypes.Contains(DescriptorType));
@@ -822,54 +864,56 @@ void FVulkanBindlessDescriptorManager::UpdateDescriptor(const FVulkanContextArra
 	}
 }
 
-void FVulkanBindlessDescriptorManager::UpdateSampler(const FVulkanContextArray& Contexts, FRHIDescriptorHandle DescriptorHandle, VkSampler VulkanSampler)
+void FVulkanBindlessDescriptorManager::UpdateDescriptor(const FVulkanContextArray& Contexts, FRHIDescriptorHandle DescriptorHandle, VkSampler VulkanSampler)
 {
-	if (bIsSupported)
+	if (IsSupported())
 	{
 		VkDescriptorDataEXT DescriptorData;
 		DescriptorData.pSampler = &VulkanSampler;
-		UpdateDescriptor(Contexts, DescriptorHandle, VK_DESCRIPTOR_TYPE_SAMPLER, DescriptorData, true);
+		UpdateDescriptorInternal(Contexts, DescriptorHandle, DescriptorData, true);
 	}
 }
 
-void FVulkanBindlessDescriptorManager::UpdateImage(const FVulkanContextArray& Contexts, FRHIDescriptorHandle DescriptorHandle, VkDescriptorType DescriptorType, VkImageView ImageView, bool bIsDepthStencil, bool bImmediateUpdate)
+void FVulkanBindlessDescriptorManager::UpdateDescriptor(const FVulkanContextArray& Contexts, FRHIDescriptorHandle DescriptorHandle, VkImageView ImageView, bool bIsDepthStencil, bool bImmediateUpdate)
 {
-	if (bIsSupported)
+	if (IsSupported())
 	{
-		check((DescriptorType == VK_DESCRIPTOR_TYPE_SAMPLED_IMAGE) || (DescriptorType == VK_DESCRIPTOR_TYPE_STORAGE_IMAGE));
+		check(DescriptorHandle.GetType() == ERHIDescriptorType::TextureSRV || DescriptorHandle.GetType() == ERHIDescriptorType::TextureUAV);
 
 		VkDescriptorImageInfo DescriptorImageInfo;
 		DescriptorImageInfo.sampler = VK_NULL_HANDLE;
 		DescriptorImageInfo.imageView = ImageView;
-		DescriptorImageInfo.imageLayout = (DescriptorType == VK_DESCRIPTOR_TYPE_STORAGE_IMAGE) ? VK_IMAGE_LAYOUT_GENERAL :
-			(bIsDepthStencil ? VK_IMAGE_LAYOUT_READ_ONLY_OPTIMAL : VK_IMAGE_LAYOUT_SHADER_READ_ONLY_OPTIMAL);
+		DescriptorImageInfo.imageLayout =
+			DescriptorHandle.GetType() == ERHIDescriptorType::TextureUAV
+			? VK_IMAGE_LAYOUT_GENERAL
+			: (bIsDepthStencil ? VK_IMAGE_LAYOUT_READ_ONLY_OPTIMAL : VK_IMAGE_LAYOUT_SHADER_READ_ONLY_OPTIMAL);
 
 		VkDescriptorDataEXT DescriptorData;
 		DescriptorData.pSampledImage = &DescriptorImageInfo;  // same pointer for storage, it's a union
-		UpdateDescriptor(Contexts, DescriptorHandle, DescriptorType, DescriptorData, bImmediateUpdate);
+		UpdateDescriptorInternal(Contexts, DescriptorHandle, DescriptorData, bImmediateUpdate);
 	}
 }
 
-void FVulkanBindlessDescriptorManager::UpdateBuffer(const FVulkanContextArray& Contexts, FRHIDescriptorHandle DescriptorHandle, VkDescriptorType DescriptorType, VkBuffer VulkanBuffer, VkDeviceSize BufferOffset, VkDeviceSize BufferSize, bool bImmediateUpdate)
+void FVulkanBindlessDescriptorManager::UpdateDescriptor(const FVulkanContextArray& Contexts, FRHIDescriptorHandle DescriptorHandle, VkBuffer VulkanBuffer, VkDeviceSize BufferOffset, VkDeviceSize BufferSize, bool bImmediateUpdate)
 {
-	if (bIsSupported)
+	if (IsSupported())
 	{
-		check((DescriptorType == VK_DESCRIPTOR_TYPE_UNIFORM_BUFFER) || (DescriptorType == VK_DESCRIPTOR_TYPE_STORAGE_BUFFER));
+		check(DescriptorHandle.GetType() == ERHIDescriptorType::BufferSRV || DescriptorHandle.GetType() == ERHIDescriptorType::BufferUAV);
 
 		VkBufferDeviceAddressInfo BufferInfo;
 		ZeroVulkanStruct(BufferInfo, VK_STRUCTURE_TYPE_BUFFER_DEVICE_ADDRESS_INFO);
 		BufferInfo.buffer = VulkanBuffer;
 		const VkDeviceAddress BufferAddress = VulkanRHI::vkGetBufferDeviceAddressKHR(Device.GetHandle(), &BufferInfo);
 
-		UpdateBuffer(Contexts, DescriptorHandle, DescriptorType, BufferAddress + BufferOffset, BufferSize, bImmediateUpdate);
+		UpdateDescriptor(Contexts, DescriptorHandle, BufferAddress + BufferOffset, BufferSize, bImmediateUpdate);
 	}
 }
 
-void FVulkanBindlessDescriptorManager::UpdateBuffer(const FVulkanContextArray& Contexts, FRHIDescriptorHandle DescriptorHandle, VkDescriptorType DescriptorType, VkDeviceAddress BufferAddress, VkDeviceSize BufferSize, bool bImmediateUpdate)
+void FVulkanBindlessDescriptorManager::UpdateDescriptor(const FVulkanContextArray& Contexts, FRHIDescriptorHandle DescriptorHandle, VkDeviceAddress BufferAddress, VkDeviceSize BufferSize, bool bImmediateUpdate)
 {
-	if (bIsSupported)
+	if (IsSupported())
 	{
-		check((DescriptorType == VK_DESCRIPTOR_TYPE_UNIFORM_BUFFER) || (DescriptorType == VK_DESCRIPTOR_TYPE_STORAGE_BUFFER));
+		check(DescriptorHandle.GetType() == ERHIDescriptorType::BufferSRV || DescriptorHandle.GetType() == ERHIDescriptorType::BufferUAV);
 
 		VkDescriptorAddressInfoEXT AddressInfo;
 		ZeroVulkanStruct(AddressInfo, VK_STRUCTURE_TYPE_DESCRIPTOR_ADDRESS_INFO_EXT);
@@ -878,15 +922,15 @@ void FVulkanBindlessDescriptorManager::UpdateBuffer(const FVulkanContextArray& C
 
 		VkDescriptorDataEXT DescriptorData;
 		DescriptorData.pStorageBuffer = &AddressInfo;  // same pointer for uniform, it's a union
-		UpdateDescriptor(Contexts, DescriptorHandle, DescriptorType, DescriptorData, bImmediateUpdate);
+		UpdateDescriptorInternal(Contexts, DescriptorHandle, DescriptorData, bImmediateUpdate);
 	}
 }
 
-void FVulkanBindlessDescriptorManager::UpdateTexelBuffer(const FVulkanContextArray& Contexts, FRHIDescriptorHandle DescriptorHandle, VkDescriptorType DescriptorType, const VkBufferViewCreateInfo& ViewInfo, bool bImmediateUpdate)
+void FVulkanBindlessDescriptorManager::UpdateDescriptor(const FVulkanContextArray& Contexts, FRHIDescriptorHandle DescriptorHandle, const VkBufferViewCreateInfo& ViewInfo, bool bImmediateUpdate)
 {
-	if (bIsSupported)
+	if (IsSupported())
 	{
-		check((DescriptorType == VK_DESCRIPTOR_TYPE_UNIFORM_TEXEL_BUFFER) || (DescriptorType == VK_DESCRIPTOR_TYPE_STORAGE_TEXEL_BUFFER));
+		check(DescriptorHandle.GetType() == ERHIDescriptorType::TypedBufferSRV || DescriptorHandle.GetType() == ERHIDescriptorType::TypedBufferUAV);
 
 		// :todo-jn: start caching buffer addresses in resources to avoid the extra call
 		VkBufferDeviceAddressInfo BufferInfo;
@@ -902,13 +946,13 @@ void FVulkanBindlessDescriptorManager::UpdateTexelBuffer(const FVulkanContextArr
 
 		VkDescriptorDataEXT DescriptorData;
 		DescriptorData.pUniformTexelBuffer = &AddressInfo;  // same pointer for storage, it's a union
-		UpdateDescriptor(Contexts, DescriptorHandle, DescriptorType, DescriptorData, bImmediateUpdate);
+		UpdateDescriptorInternal(Contexts, DescriptorHandle, DescriptorData, bImmediateUpdate);
 	}
 }
 
-void FVulkanBindlessDescriptorManager::UpdateAccelerationStructure(const FVulkanContextArray& Contexts, FRHIDescriptorHandle DescriptorHandle, VkAccelerationStructureKHR AccelerationStructure, bool bImmediateUpdate)
+void FVulkanBindlessDescriptorManager::UpdateDescriptor(const FVulkanContextArray& Contexts, FRHIDescriptorHandle DescriptorHandle, VkAccelerationStructureKHR AccelerationStructure, bool bImmediateUpdate)
 {
-	if (bIsSupported)
+	if (IsSupported())
 	{
 		// :todo-jn: start caching AccelerationStructure in resources to avoid the extra call
 		VkAccelerationStructureDeviceAddressInfoKHR AccelerationStructureDeviceAddressInfo;
@@ -918,6 +962,6 @@ void FVulkanBindlessDescriptorManager::UpdateAccelerationStructure(const FVulkan
 
 		VkDescriptorDataEXT DescriptorData;
 		DescriptorData.accelerationStructure = BufferAddress;
-		UpdateDescriptor(Contexts, DescriptorHandle, VK_DESCRIPTOR_TYPE_ACCELERATION_STRUCTURE_KHR, DescriptorData, bImmediateUpdate);
+		UpdateDescriptorInternal(Contexts, DescriptorHandle, DescriptorData, bImmediateUpdate);
 	}
 }
diff --git a/Engine/Source/Runtime/VulkanRHI/Private/VulkanBindlessDescriptorManager.h b/Engine/Source/Runtime/VulkanRHI/Private/VulkanBindlessDescriptorManager.h
--- a/Engine/Source/Runtime/VulkanRHI/Private/VulkanBindlessDescriptorManager.h
+++ b/Engine/Source/Runtime/VulkanRHI/Private/VulkanBindlessDescriptorManager.h
@@ -30,9 +30,38 @@ public:
 	void Init();
 	void Deinit();
 
-	bool IsSupported()
+	void BindDescriptorBuffers(VkCommandBuffer CommandBuffer, VkPipelineStageFlags SupportedStages);
+
+	FRHIDescriptorHandle AllocateDescriptor(ERHIDescriptorType DescriptorType);
+	TOptional<FRHIDescriptorAllocation> AllocateDescriptors(ERHIDescriptorType DescriptorType, uint32 DescriptorCount);
+
+	FRHIDescriptorHandle AllocateDescriptor(ERHIDescriptorType DescriptorType, VkDeviceAddress BufferAddress, VkDeviceSize BufferSize);
+	FRHIDescriptorHandle AllocateDescriptor(ERHIDescriptorType DescriptorType, FVulkanBuffer* Buffer, uint32 ExtraOffset);
+
+	FRHIDescriptorHandle AllocateDescriptor(VkSampler VulkanSampler);
+
+	void FreeDescriptor(FRHIDescriptorHandle DescriptorHandle);
+	void FreeDescriptors(FRHIDescriptorAllocation Allocation, VkDescriptorType DescriptorType);
+
+	void UpdateDescriptor(const FVulkanContextArray& Contexts, FRHIDescriptorHandle DescriptorHandle, VkSampler VulkanSampler);
+	void UpdateDescriptor(const FVulkanContextArray& Contexts, FRHIDescriptorHandle DescriptorHandle, VkImageView VulkanImage, bool bIsDepthStencil, bool bImmediateUpdate = true);
+
+	void UpdateDescriptor(const FVulkanContextArray& Contexts, FRHIDescriptorHandle DescriptorHandle, VkBuffer VulkanBuffer, VkDeviceSize BufferOffset, VkDeviceSize BufferSize, bool bImmediateUpdate = true);
+	void UpdateDescriptor(const FVulkanContextArray& Contexts, FRHIDescriptorHandle DescriptorHandle, VkDeviceAddress BufferAddress, VkDeviceSize BufferSize, bool bImmediateUpdate = true);
+
+	void UpdateDescriptor(const FVulkanContextArray& Contexts, FRHIDescriptorHandle DescriptorHandle, const VkBufferViewCreateInfo& ViewInfo, bool bImmediateUpdate = true);
+	void UpdateDescriptor(const FVulkanContextArray& Contexts, FRHIDescriptorHandle DescriptorHandle, VkAccelerationStructureKHR AccelerationStructure, bool bImmediateUpdate = true);
+
+	void RegisterUniformBuffers(FVulkanCommandListContext& Context, VkShaderStageFlags StageFlags, TConstArrayView<VkWriteDescriptorSet> UBDescriptors);
+
+	ERHIBindlessConfiguration GetConfiguration() const
 	{
-		return bIsSupported;
+		return Configuration;
+	}
+
+	bool IsSupported() const
+	{
+		return !IsBindlessDisabled(Configuration);
 	}
 
 	VkPipelineLayout GetPipelineLayout() const
@@ -51,30 +80,11 @@ public:
 		return PushConstantRanges;
 	}
 
-	void BindDescriptorBuffers(VkCommandBuffer CommandBuffer, VkPipelineStageFlags SupportedStages);
-
-	FRHIDescriptorHandle AllocateDescriptor(VkDescriptorType DescriptorType);
-	FRHIDescriptorHandle AllocateDescriptor(VkDescriptorType DescriptorType, VkDeviceAddress BufferAddress, VkDeviceSize BufferSize);
-	FRHIDescriptorHandle AllocateDescriptor(VkDescriptorType DescriptorType, FVulkanBuffer* Buffer, uint32 ExtraOffset);
-
-	FRHIDescriptorHandle AllocateDescriptor(VkSampler VulkanSampler);
-
-	void FreeDescriptor(FRHIDescriptorHandle DescriptorHandle);
-
-	void UpdateSampler(const FVulkanContextArray& Contexts, FRHIDescriptorHandle DescriptorHandle, VkSampler VulkanSampler);
-	void UpdateImage(const FVulkanContextArray& Contexts, FRHIDescriptorHandle DescriptorHandle, VkDescriptorType DescriptorType, VkImageView VulkanImage, bool bIsDepthStencil, bool bImmediateUpdate = true);
-
-	void UpdateBuffer(const FVulkanContextArray& Contexts, FRHIDescriptorHandle DescriptorHandle, VkDescriptorType DescriptorType, VkBuffer VulkanBuffer, VkDeviceSize BufferOffset, VkDeviceSize BufferSize, bool bImmediateUpdate = true);
-	void UpdateBuffer(const FVulkanContextArray& Contexts, FRHIDescriptorHandle DescriptorHandle, VkDescriptorType DescriptorType, VkDeviceAddress BufferAddress, VkDeviceSize BufferSize, bool bImmediateUpdate = true);
-
-	void UpdateTexelBuffer(const FVulkanContextArray& Contexts, FRHIDescriptorHandle DescriptorHandle, VkDescriptorType DescriptorType, const VkBufferViewCreateInfo& ViewInfo, bool bImmediateUpdate = true);
-	void UpdateAccelerationStructure(const FVulkanContextArray& Contexts, FRHIDescriptorHandle DescriptorHandle, VkAccelerationStructureKHR AccelerationStructure, bool bImmediateUpdate = true);
-
-	void RegisterUniformBuffers(FVulkanCommandListContext& Context, VkShaderStageFlags StageFlags, TConstArrayView<VkWriteDescriptorSet> UBDescriptors);
-
 private:
 	FVulkanDevice& Device;
-	const bool bIsSupported;
+
+	ERHIBindlessConfiguration Configuration{};
+	ERHIPipeline              ConfiguredPipelines = ERHIPipeline::None;
 
 	VkDescriptorSetLayout EmptyDescriptorSetLayout = VK_NULL_HANDLE;
 
@@ -116,6 +126,5 @@ private:
 	BindlessSetState& GetBindlessState(ERHIDescriptorType DescriptorType);
 	BindlessSetState& GetBindlessState(VkDescriptorType DescriptorType);
 
-	void UpdateDescriptor(const FVulkanContextArray& Contexts, FRHIDescriptorHandle DescriptorHandle, VkDescriptorType DescriptorType, VkDescriptorDataEXT DescriptorData, bool bImmediateUpdate);
-	static bool VerifySupport(FVulkanDevice& InDevice);
+	void UpdateDescriptorInternal(const FVulkanContextArray& Contexts, FRHIDescriptorHandle DescriptorHandle, VkDescriptorDataEXT DescriptorData, bool bImmediateUpdate);
 };
diff --git a/Engine/Source/Runtime/VulkanRHI/Private/VulkanRayTracing.cpp b/Engine/Source/Runtime/VulkanRHI/Private/VulkanRayTracing.cpp
--- a/Engine/Source/Runtime/VulkanRHI/Private/VulkanRayTracing.cpp
+++ b/Engine/Source/Runtime/VulkanRHI/Private/VulkanRayTracing.cpp
@@ -556,7 +556,7 @@ void FVulkanRayTracingGeometry::SetupHitGroupSystemParameters()
 	{
 		if (Buffer)
 		{
-			return BindlessDescriptorManager->AllocateDescriptor(VK_DESCRIPTOR_TYPE_STORAGE_BUFFER, Buffer, ExtraOffset);
+			return BindlessDescriptorManager->AllocateDescriptor(ERHIDescriptorType::BufferUAV, Buffer, ExtraOffset);
 		}
 		return FRHIDescriptorHandle();
 	};
@@ -734,7 +734,7 @@ void FVulkanRayTracingScene::BindBuffer(FRHIBuffer* InBuffer, uint32 InBufferOff
 
 		check(InBufferOffset % GRHIRayTracingAccelerationStructureAlignment == 0);
 
-		View = MakeUnique<FVulkanView>(Device, VK_DESCRIPTOR_TYPE_ACCELERATION_STRUCTURE_KHR);
+		View = MakeUnique<FVulkanView>(Device, ERHIDescriptorType::AccelerationStructure);
 		View->InitAsAccelerationStructureView(
 			{},
 			AccelerationStructureBuffer
diff --git a/Engine/Source/Runtime/VulkanRHI/Private/VulkanResources.h b/Engine/Source/Runtime/VulkanRHI/Private/VulkanResources.h
--- a/Engine/Source/Runtime/VulkanRHI/Private/VulkanResources.h
+++ b/Engine/Source/Runtime/VulkanRHI/Private/VulkanResources.h
@@ -430,8 +430,6 @@ struct FVulkanTextureViewDesc
 	uint32 GetHash() const;
 };
 
-VkDescriptorType GetVkDescriptorTypeForViewDesc(const FRHIViewDesc& ViewDesc);
-
 class FVulkanView
 {
 public:
@@ -485,7 +483,7 @@ public:
 		AccelerationStructure = TStorage::IndexOfType<FAccelerationStructureView>(),
 	};
 
-	FVulkanView(FVulkanDevice& InDevice, VkDescriptorType InDescriptorType);
+	FVulkanView(FVulkanDevice& InDevice, ERHIDescriptorType InDescriptorType);
 
 	~FVulkanView();
 
@@ -527,25 +525,20 @@ public:
 		return BindlessHandle;
 	}
 
-	VkDescriptorType GetDescriptorType() const
-	{
-		return DescriptorType;
-	}
-
 protected:
 	void InitAsTextureView(const FVulkanContextArray& Contexts, VkImage InImage, const FVulkanTextureViewDesc& ViewDesc);
 
 private:
 	FVulkanDevice& Device;
 	FRHIDescriptorHandle BindlessHandle;
-	const VkDescriptorType DescriptorType;
+	const ERHIDescriptorType DescriptorType;
 	TStorage Storage;
 };
 
 class FVulkanLinkedView : public FVulkanView, public TIntrusiveLinkedList<FVulkanLinkedView>
 {
 protected:
-	FVulkanLinkedView(FVulkanDevice& Device, VkDescriptorType DescriptorType)
+	FVulkanLinkedView(FVulkanDevice& Device, ERHIDescriptorType DescriptorType)
 		: FVulkanView(Device, DescriptorType)
 	{}
 
diff --git a/Engine/Source/Runtime/VulkanRHI/Private/VulkanSRV.cpp b/Engine/Source/Runtime/VulkanRHI/Private/VulkanSRV.cpp
--- a/Engine/Source/Runtime/VulkanRHI/Private/VulkanSRV.cpp
+++ b/Engine/Source/Runtime/VulkanRHI/Private/VulkanSRV.cpp
@@ -20,9 +20,40 @@ static VkImageViewType GetVkImageViewTypeForDimensionSRV(FRHIViewDesc::EDimensio
 	return VK_IMAGE_VIEW_TYPE_MAX_ENUM;
 }
 
+inline ERHIDescriptorType GetSRVDescriptorType(const FRHIViewDesc& ViewDesc)
+{
+	if (ViewDesc.IsTexture())
+	{
+		// Sampled images aren't supported in R64, shadercompiler patches them to storage image
+		if (ViewDesc.Texture.SRV.Format == PF_R64_UINT)
+		{
+			return ERHIDescriptorType::TextureUAV;
+		}
+
+		return ERHIDescriptorType::TextureSRV;
+	}
+
+	switch (ViewDesc.Buffer.SRV.BufferType)
+	{
+	case FRHIViewDesc::EBufferType::Raw:
+	case FRHIViewDesc::EBufferType::Structured:
+		return ERHIDescriptorType::BufferUAV;
+
+	case FRHIViewDesc::EBufferType::Typed:
+		return ERHIDescriptorType::TypedBufferSRV;
+
+	case FRHIViewDesc::EBufferType::AccelerationStructure:
+		return ERHIDescriptorType::AccelerationStructure;
+
+	default:
+		checkNoEntry();
+		return ERHIDescriptorType::Invalid;
+	}
+}
+
 FVulkanShaderResourceView::FVulkanShaderResourceView(FRHICommandListBase& RHICmdList, FVulkanDevice& InDevice, FRHIViewableResource* InResource, FRHIViewDesc const& InViewDesc)
 	: FRHIShaderResourceView(InResource, InViewDesc)
-	, FVulkanLinkedView(InDevice, GetVkDescriptorTypeForViewDesc(InViewDesc))
+	, FVulkanLinkedView(InDevice, GetSRVDescriptorType(InViewDesc))
 {
 	RHICmdList.EnqueueLambda([this](FRHICommandListBase&)
 	{
diff --git a/Engine/Source/Runtime/VulkanRHI/Private/VulkanTexture.cpp b/Engine/Source/Runtime/VulkanRHI/Private/VulkanTexture.cpp
--- a/Engine/Source/Runtime/VulkanRHI/Private/VulkanTexture.cpp
+++ b/Engine/Source/Runtime/VulkanRHI/Private/VulkanTexture.cpp
@@ -2196,7 +2196,6 @@ FVulkanTexture::FVulkanTexture(FVulkanDevice& InDevice, const FRHITextureCreateD
 static FVulkanTextureViewDesc CreateTextureViewDesc(FVulkanTexture* Texture, VkImageUsageFlags ViewUsageFlags, VkSamplerYcbcrConversion SamplerYcbcrConversion)
 {
 	const FRHITextureDesc& Desc = Texture->GetDesc();
-	const VkDescriptorType DescriptorType = Texture->SupportsSampling() ? VK_DESCRIPTOR_TYPE_SAMPLED_IMAGE : VK_DESCRIPTOR_TYPE_STORAGE_IMAGE;
 	const bool bHasUAVFormat = (Desc.UAVFormat != PF_Unknown) && (Desc.UAVFormat != Desc.Format);
 
 	if (EnumHasAllFlags(Desc.Flags, ETextureCreateFlags::SRGB) || (bHasUAVFormat && !UE::PixelFormat::HasCapabilities(Desc.Format, EPixelFormatCapabilities::UAV)))
@@ -2205,7 +2204,7 @@ static FVulkanTextureViewDesc CreateTextureViewDesc(FVulkanTexture* Texture, VkI
 	}
 
 	const bool bUseIdentitySwizzle =
-		(DescriptorType != VK_DESCRIPTOR_TYPE_SAMPLED_IMAGE) ||
+		!Texture->SupportsSampling() ||
 		(Texture->ViewFormat == VK_FORMAT_UNDEFINED); // External buffer textures also require identity swizzle
 
 	FVulkanTextureViewDesc ViewDesc;
@@ -2229,7 +2228,7 @@ void FVulkanTexture::InitViews(VkImageUsageFlags ViewUsageFlags, VkSamplerYcbcrC
 	if (Image != VK_NULL_HANDLE && GetViewType() != VK_IMAGE_VIEW_TYPE_MAX_ENUM)
 	{
 		const FVulkanTextureViewDesc ViewDesc = CreateTextureViewDesc(this, ViewUsageFlags, SamplerYcbcrConversion);
-		const VkDescriptorType DescriptorType = SupportsSampling() ? VK_DESCRIPTOR_TYPE_SAMPLED_IMAGE : VK_DESCRIPTOR_TYPE_STORAGE_IMAGE;
+		const ERHIDescriptorType DescriptorType = SupportsSampling() ? ERHIDescriptorType::TextureSRV : ERHIDescriptorType::TextureUAV;
 
 		if (DefaultView)
 		{
@@ -2273,7 +2272,7 @@ void FVulkanTexture::UpdateViews(const FVulkanContextArray& Contexts, VkImageUsa
 	if (Image != VK_NULL_HANDLE && GetViewType() != VK_IMAGE_VIEW_TYPE_MAX_ENUM)
 	{
 		const FVulkanTextureViewDesc ViewDesc = CreateTextureViewDesc(this, ViewUsageFlags, SamplerYcbcrConversion);
-		const VkDescriptorType DescriptorType = SupportsSampling() ? VK_DESCRIPTOR_TYPE_SAMPLED_IMAGE : VK_DESCRIPTOR_TYPE_STORAGE_IMAGE;
+		const ERHIDescriptorType DescriptorType = SupportsSampling() ? ERHIDescriptorType::TextureSRV : ERHIDescriptorType::TextureUAV;
 
 		if (DefaultView)
 		{
@@ -2332,7 +2331,7 @@ FVulkanView* FVulkanTexture::FindOrAddInternalView(const FVulkanTextureViewDesc&
 	FVulkanView*& VulkanView = InternalViews.FindOrAdd(ViewDesc.GetHash(), nullptr);
 	if (!VulkanView)
 	{
-		VulkanView = new FVulkanView(*Device, VK_DESCRIPTOR_TYPE_SAMPLED_IMAGE);
+		VulkanView = new FVulkanView(*Device, ERHIDescriptorType::TextureSRV);
 		VulkanView->CreateTextureView(Image, ViewDesc);
 	}
 	return VulkanView;
diff --git a/Engine/Source/Runtime/VulkanRHI/Private/VulkanUAV.cpp b/Engine/Source/Runtime/VulkanRHI/Private/VulkanUAV.cpp
--- a/Engine/Source/Runtime/VulkanRHI/Private/VulkanUAV.cpp
+++ b/Engine/Source/Runtime/VulkanRHI/Private/VulkanUAV.cpp
@@ -25,9 +25,34 @@ static VkImageViewType GetVkImageViewTypeForDimensionUAV(FRHIViewDesc::EDimensio
 
 }
 
+inline ERHIDescriptorType GetUAVDescriptorType(const FRHIViewDesc& ViewDesc)
+{
+	if (ViewDesc.IsTexture())
+	{
+		return ERHIDescriptorType::TextureUAV;
+	}
+
+	switch (ViewDesc.Buffer.UAV.BufferType)
+	{
+	case FRHIViewDesc::EBufferType::Raw:
+	case FRHIViewDesc::EBufferType::Structured:
+		return ERHIDescriptorType::BufferUAV;
+
+	case FRHIViewDesc::EBufferType::Typed:
+		return ERHIDescriptorType::TypedBufferUAV;
+
+	case FRHIViewDesc::EBufferType::AccelerationStructure:
+		return ERHIDescriptorType::AccelerationStructure;
+
+	default:
+		checkNoEntry();
+		return ERHIDescriptorType::Invalid;
+	}
+}
+
 FVulkanUnorderedAccessView::FVulkanUnorderedAccessView(FRHICommandListBase& RHICmdList, FVulkanDevice& InDevice, FRHIViewableResource* InResource, FRHIViewDesc const& InViewDesc)
 	: FRHIUnorderedAccessView(InResource, InViewDesc)
-	, FVulkanLinkedView(InDevice, GetVkDescriptorTypeForViewDesc(InViewDesc))
+	, FVulkanLinkedView(InDevice, GetUAVDescriptorType(InViewDesc))
 {
 	RHICmdList.EnqueueLambda([this](FRHICommandListBase&)
 	{
diff --git a/Engine/Source/Runtime/VulkanRHI/Private/VulkanUniformBuffer.cpp b/Engine/Source/Runtime/VulkanRHI/Private/VulkanUniformBuffer.cpp
--- a/Engine/Source/Runtime/VulkanRHI/Private/VulkanUniformBuffer.cpp
+++ b/Engine/Source/Runtime/VulkanRHI/Private/VulkanUniformBuffer.cpp
@@ -230,7 +230,7 @@ void FVulkanUniformBuffer::UpdateBindlessHandle()
 				Device->GetDeferredDeletionQueue().EnqueueBindlessHandle(BindlessHandle);
 			}
 
-			BindlessHandle = Device->GetBindlessDescriptorManager()->AllocateDescriptor(VK_DESCRIPTOR_TYPE_UNIFORM_BUFFER, CurrentAddress, GetSize());
+			BindlessHandle = Device->GetBindlessDescriptorManager()->AllocateDescriptor(ERHIDescriptorType::BufferSRV, CurrentAddress, GetSize());
 
 			CachedDeviceAddress = CurrentAddress;
 		}
diff --git a/Engine/Source/Runtime/VulkanRHI/Private/VulkanView.cpp b/Engine/Source/Runtime/VulkanRHI/Private/VulkanView.cpp
--- a/Engine/Source/Runtime/VulkanRHI/Private/VulkanView.cpp
+++ b/Engine/Source/Runtime/VulkanRHI/Private/VulkanView.cpp
@@ -23,69 +23,7 @@ uint32 FVulkanTextureViewDesc::GetHash() const
 	return Hash;
 }
 
-VkDescriptorType GetVkDescriptorTypeForViewDesc(const FRHIViewDesc& ViewDesc)
-{
-	if (ViewDesc.IsBuffer())
-	{
-		if (ViewDesc.IsSRV())
-		{
-			switch (ViewDesc.Buffer.SRV.BufferType)
-			{
-			case FRHIViewDesc::EBufferType::Raw:
-			case FRHIViewDesc::EBufferType::Structured:
-				return VK_DESCRIPTOR_TYPE_STORAGE_BUFFER;
-
-			case FRHIViewDesc::EBufferType::Typed:
-				return VK_DESCRIPTOR_TYPE_UNIFORM_TEXEL_BUFFER;
-
-			case FRHIViewDesc::EBufferType::AccelerationStructure:
-				return VK_DESCRIPTOR_TYPE_ACCELERATION_STRUCTURE_KHR;
-
-			default:
-				checkNoEntry();
-				return VK_DESCRIPTOR_TYPE_MAX_ENUM;
-			}
-		}
-		else
-		{
-			switch (ViewDesc.Buffer.UAV.BufferType)
-			{
-			case FRHIViewDesc::EBufferType::Raw:
-			case FRHIViewDesc::EBufferType::Structured:
-				return VK_DESCRIPTOR_TYPE_STORAGE_BUFFER;
-
-			case FRHIViewDesc::EBufferType::Typed:
-				return VK_DESCRIPTOR_TYPE_STORAGE_TEXEL_BUFFER;
-
-			case FRHIViewDesc::EBufferType::AccelerationStructure:
-				return VK_DESCRIPTOR_TYPE_ACCELERATION_STRUCTURE_KHR;
-
-			default:
-				checkNoEntry();
-				return VK_DESCRIPTOR_TYPE_MAX_ENUM;
-			}
-		}
-	}
-	else
-	{
-		if (ViewDesc.IsSRV())
-		{
-			// Sampled images aren't supported in R64, shadercompiler patches them to storage image
-			if (ViewDesc.Texture.SRV.Format == PF_R64_UINT)
-			{
-				return VK_DESCRIPTOR_TYPE_STORAGE_IMAGE;
-			}
-
-			return VK_DESCRIPTOR_TYPE_SAMPLED_IMAGE;
-		}
-		else
-		{
-			return VK_DESCRIPTOR_TYPE_STORAGE_IMAGE;
-		}
-	}
-}
-
-FVulkanView::FVulkanView(FVulkanDevice& InDevice, VkDescriptorType InDescriptorType)
+FVulkanView::FVulkanView(FVulkanDevice& InDevice, ERHIDescriptorType InDescriptorType)
 	: Device(InDevice)
 	, DescriptorType(InDescriptorType)
 {
@@ -185,7 +123,7 @@ void FVulkanView::InitAsTypedBufferView(const FVulkanContextArray& Contexts, FVu
 	INC_DWORD_STAT(STAT_VulkanNumBufferViews);
 	// :todo-jn: the buffer view is actually not needed in bindless anymore
 
-	Device.GetBindlessDescriptorManager()->UpdateTexelBuffer(Contexts, BindlessHandle, DescriptorType, ViewInfo, bImmediateUpdate);
+	Device.GetBindlessDescriptorManager()->UpdateDescriptor(Contexts, BindlessHandle, ViewInfo, bImmediateUpdate);
 }
 
 void FVulkanView::InitAsTextureView(const FVulkanContextArray& Contexts, VkImage InImage, const FVulkanTextureViewDesc& ViewDesc)
@@ -278,7 +216,7 @@ void FVulkanView::InitAsTextureView(const FVulkanContextArray& Contexts, VkImage
 	}
 
 	const bool bDepthOrStencilAspect = (ViewDesc.AspectFlags & (VK_IMAGE_ASPECT_DEPTH_BIT | VK_IMAGE_ASPECT_STENCIL_BIT)) != 0;
-	Device.GetBindlessDescriptorManager()->UpdateImage(Contexts, BindlessHandle, DescriptorType, TV.View, bDepthOrStencilAspect, bImmediateUpdate);
+	Device.GetBindlessDescriptorManager()->UpdateDescriptor(Contexts, BindlessHandle, TV.View, bDepthOrStencilAspect, bImmediateUpdate);
 }
 
 void FVulkanView::CreateTextureView(VkImage InImage, const FVulkanTextureViewDesc& ViewDesc)
@@ -314,7 +252,7 @@ void FVulkanView::InitAsStructuredBufferView(const FVulkanContextArray& Contexts
 
 	SBV.Size = InSize;
 
-	Device.GetBindlessDescriptorManager()->UpdateBuffer(Contexts, BindlessHandle, DescriptorType, Buffer->GetHandle(), TotalOffset, InSize, bImmediateUpdate);
+	Device.GetBindlessDescriptorManager()->UpdateDescriptor(Contexts, BindlessHandle, Buffer->GetHandle(), TotalOffset, InSize, bImmediateUpdate);
 }
 
 void FVulkanView::InitAsAccelerationStructureView(const FVulkanContextArray& Contexts, FVulkanBuffer* Buffer, uint32 Offset, uint32 Size)
@@ -332,7 +270,7 @@ void FVulkanView::InitAsAccelerationStructureView(const FVulkanContextArray& Con
 
 	VERIFYVULKANRESULT(VulkanDynamicAPI::vkCreateAccelerationStructureKHR(Device.GetHandle(), &CreateInfo, VULKAN_CPU_ALLOCATOR, &ASV.Handle));
 
-	Device.GetBindlessDescriptorManager()->UpdateAccelerationStructure(Contexts, BindlessHandle, ASV.Handle, true);
+	Device.GetBindlessDescriptorManager()->UpdateDescriptor(Contexts, BindlessHandle, ASV.Handle, true);
 }
 
 //
diff --git a/Engine/Source/Runtime/VulkanRHI/Private/VulkanViewport.cpp b/Engine/Source/Runtime/VulkanRHI/Private/VulkanViewport.cpp
--- a/Engine/Source/Runtime/VulkanRHI/Private/VulkanViewport.cpp
+++ b/Engine/Source/Runtime/VulkanRHI/Private/VulkanViewport.cpp
@@ -400,7 +400,7 @@ void FVulkanViewport::InitImages(FVulkanContextCommon& Context, TConstArrayView<
 		ViewDesc.ArraySliceIndex = 0;
 		ViewDesc.NumArraySlices = 1;
 
-		FVulkanView* NewView = new FVulkanView(Device, VK_DESCRIPTOR_TYPE_SAMPLED_IMAGE);
+		FVulkanView* NewView = new FVulkanView(Device, ERHIDescriptorType::TextureSRV);
 		NewView->CreateTextureView(Images[Index], ViewDesc);
 
 		TextureViews.Add(NewView);
