VulkanRHI: Fix validation error when a VulkanView is created for a framebuffer (don't automatically create a bindless SRV handle).

#rnx

#Synced-CL 49728766

diff --git a/Engine/Source/Runtime/VulkanRHI/Private/VulkanTexture.cpp b/Engine/Source/Runtime/VulkanRHI/Private/VulkanTexture.cpp
--- a/Engine/Source/Runtime/VulkanRHI/Private/VulkanTexture.cpp
+++ b/Engine/Source/Runtime/VulkanRHI/Private/VulkanTexture.cpp
@@ -2331,7 +2331,9 @@ FVulkanView* FVulkanTexture::FindOrAddInternalView(const FVulkanTextureViewDesc&
 	FVulkanView*& VulkanView = InternalViews.FindOrAdd(ViewDesc.GetHash(), nullptr);
 	if (!VulkanView)
 	{
-		VulkanView = new FVulkanView(*Device, ERHIDescriptorType::TextureSRV);
+		// When creating a view for a framebuffer, only VK_IMAGE_USAGE_COLOR_ATTACHMENT_BIT will be set
+		const ERHIDescriptorType DescriptorType = VKHasAnyFlags(ViewDesc.ImageUsageFlags, VK_IMAGE_USAGE_SAMPLED_BIT) ? ERHIDescriptorType::TextureSRV : ERHIDescriptorType::Invalid;
+		VulkanView = new FVulkanView(*Device, DescriptorType);
 		VulkanView->CreateTextureView(Image, ViewDesc);
 	}
 	return VulkanView;
diff --git a/Engine/Source/Runtime/VulkanRHI/Private/VulkanView.cpp b/Engine/Source/Runtime/VulkanRHI/Private/VulkanView.cpp
--- a/Engine/Source/Runtime/VulkanRHI/Private/VulkanView.cpp
+++ b/Engine/Source/Runtime/VulkanRHI/Private/VulkanView.cpp
@@ -27,7 +27,11 @@ FVulkanView::FVulkanView(FVulkanDevice& InDevice, ERHIDescriptorType InDescripto
 	: Device(InDevice)
 	, DescriptorType(InDescriptorType)
 {
-	BindlessHandle = Device.GetBindlessDescriptorManager()->AllocateDescriptor(InDescriptorType);
+	// Some views might not want bindless handles (like views used strictly for framebuffers)
+	if (InDescriptorType != ERHIDescriptorType::Invalid)
+	{
+		BindlessHandle = Device.GetBindlessDescriptorManager()->AllocateDescriptor(InDescriptorType);
+	}
 }
 
 FVulkanView::~FVulkanView()
@@ -215,8 +219,11 @@ void FVulkanView::InitAsTextureView(const FVulkanContextArray& Contexts, VkImage
 		TV.ViewId = ++GVulkanImageViewHandleIdCounter;
 	}
 
-	const bool bDepthOrStencilAspect = (ViewDesc.AspectFlags & (VK_IMAGE_ASPECT_DEPTH_BIT | VK_IMAGE_ASPECT_STENCIL_BIT)) != 0;
-	Device.GetBindlessDescriptorManager()->UpdateDescriptor(Contexts, BindlessHandle, TV.View, bDepthOrStencilAspect, bImmediateUpdate);
+	if (BindlessHandle.IsValid())
+	{
+		const bool bDepthOrStencilAspect = (ViewDesc.AspectFlags & (VK_IMAGE_ASPECT_DEPTH_BIT | VK_IMAGE_ASPECT_STENCIL_BIT)) != 0;
+		Device.GetBindlessDescriptorManager()->UpdateDescriptor(Contexts, BindlessHandle, TV.View, bDepthOrStencilAspect, bImmediateUpdate);
+	}
 }
 
 void FVulkanView::CreateTextureView(VkImage InImage, const FVulkanTextureViewDesc& ViewDesc)
