VulkanRHI: Use NeedsExtraTransitions to skip barriers between copies when they are batched

#rnx

#Synced-CL 49763429

diff --git a/Engine/Source/Runtime/VulkanRHI/Private/VulkanBuffer.cpp b/Engine/Source/Runtime/VulkanRHI/Private/VulkanBuffer.cpp
--- a/Engine/Source/Runtime/VulkanRHI/Private/VulkanBuffer.cpp
+++ b/Engine/Source/Runtime/VulkanRHI/Private/VulkanBuffer.cpp
@@ -473,7 +473,8 @@ void FVulkanBuffer::Unlock(FRHICommandListBase& RHICmdList)
 
 		FVulkanPendingBufferLock PendingLock = GetPendingBufferLock(this);
 
-		RHICmdList.EnqueueLambda(TEXT("FVulkanBuffer::Unlock"), [Buffer=this, PendingLock](FRHICommandListBase& CmdList)
+		const bool bNeedsTransitions = RHICmdList.NeedsExtraTransitions();
+		RHICmdList.EnqueueLambda(TEXT("FVulkanBuffer::Unlock"), [Buffer=this, PendingLock, bNeedsTransitions](FRHICommandListBase& ExecutingCmdList)
 		{
 			VulkanRHI::FStagingBuffer* StagingBuffer = PendingLock.StagingBuffer;
 			check(StagingBuffer);
@@ -492,11 +493,11 @@ void FVulkanBuffer::Unlock(FRHICommandListBase& RHICmdList)
 				FVulkanContextCommon* CommonContext = nullptr;
 				if (PendingLock.FirstLock)
 				{
-					CommonContext = &FVulkanUploadContext::Get(CmdList);
+					CommonContext = &FVulkanUploadContext::Get(ExecutingCmdList);
 				}
 				else
 				{
-					CommonContext = &FVulkanCommandListContext::Get(CmdList);
+					CommonContext = &FVulkanCommandListContext::Get(ExecutingCmdList);
 				}
 				FVulkanCommandBuffer& CommandBuffer = CommonContext->GetCommandBuffer();
 				check(CommandBuffer.IsOutsideRenderPass());
@@ -511,9 +512,12 @@ void FVulkanBuffer::Unlock(FRHICommandListBase& RHICmdList)
 				Region.dstOffset = PendingLock.Offset + Buffer->CurrentBufferAlloc.Alloc.Offset;
 				VulkanRHI::vkCmdCopyBuffer(CommandBufferHandle, StagingBuffer->GetHandle(), Buffer->CurrentBufferAlloc.Alloc.GetBufferHandle(), 1, &Region);
 
-				// High level code expects the data in Buffer to be ready to read
-				VkMemoryBarrier BarrierAfter = { VK_STRUCTURE_TYPE_MEMORY_BARRIER, nullptr, VK_ACCESS_TRANSFER_WRITE_BIT, VK_ACCESS_MEMORY_READ_BIT | VK_ACCESS_MEMORY_WRITE_BIT };
-				VulkanRHI::vkCmdPipelineBarrier(CommandBufferHandle, VK_PIPELINE_STAGE_TRANSFER_BIT, VK_PIPELINE_STAGE_ALL_COMMANDS_BIT, 0, 1, &BarrierAfter, 0, nullptr, 0, nullptr);
+				// High level code expects the data in Buffer to be ready to read, unless the transitions are batched in a pass with multiple copies
+				if (bNeedsTransitions)
+				{
+					VkMemoryBarrier BarrierAfter = { VK_STRUCTURE_TYPE_MEMORY_BARRIER, nullptr, VK_ACCESS_TRANSFER_WRITE_BIT, VK_ACCESS_MEMORY_READ_BIT | VK_ACCESS_MEMORY_WRITE_BIT };
+					VulkanRHI::vkCmdPipelineBarrier(CommandBufferHandle, VK_PIPELINE_STAGE_TRANSFER_BIT, VK_PIPELINE_STAGE_ALL_COMMANDS_BIT, 0, 1, &BarrierAfter, 0, nullptr, 0, nullptr);
+				}
 
 				Buffer->Device.GetStagingManager().ReleaseBuffer(CommonContext, StagingBuffer);
 			}
diff --git a/Engine/Source/Runtime/VulkanRHI/Private/VulkanResources.h b/Engine/Source/Runtime/VulkanRHI/Private/VulkanResources.h
--- a/Engine/Source/Runtime/VulkanRHI/Private/VulkanResources.h
+++ b/Engine/Source/Runtime/VulkanRHI/Private/VulkanResources.h
@@ -146,7 +146,7 @@ public:
 		return CreateHandle(Desc, Layout, LayoutHash);
 	}
 
-	inline const FString& GetDebugName() const
+	const FString& GetDebugName() const
 	{
 		return CodeHeader.DebugName;
 	}
@@ -162,7 +162,7 @@ public:
 		return CodeHeader;
 	}
 
-	inline uint64 GetShaderKey() const
+	uint64 GetShaderKey() const
 	{
 		return ShaderKey;
 	}
@@ -203,6 +203,19 @@ public:
 		return bUsesBindless;
 	}
 
+	void WaitForAsyncCompile()
+	{
+		if (AsyncCompileTask.IsValid() && !AsyncCompileTask->IsComplete())
+		{
+			AsyncCompileTask->Wait();
+		}
+	}
+
+	VkShaderEXT GetShaderObject() const
+	{
+		return ShaderObject;
+	}
+
 protected:
 
 #if UE_BUILD_DEBUG || UE_BUILD_DEVELOPMENT
@@ -778,7 +791,7 @@ public:
 	VULKANRHI_API VkDeviceMemory GetAllocationHandle() const;
 	VULKANRHI_API uint64 GetAllocationOffset() const;
 
-	static void InternalLockWrite(FVulkanContextCommon& Context, FVulkanTexture* Surface, const VkBufferImageCopy& Region, VulkanRHI::FStagingBuffer* StagingBuffer);
+	static void InternalLockWrite(FVulkanContextCommon& Context, FVulkanTexture* Surface, const VkBufferImageCopy& Region, VulkanRHI::FStagingBuffer* StagingBuffer, bool bNeedsTransitions);
 
 	const FVulkanCpuReadbackBuffer* GetCpuReadbackBuffer() const { return CpuReadbackBuffer; }
 
diff --git a/Engine/Source/Runtime/VulkanRHI/Private/VulkanTexture.cpp b/Engine/Source/Runtime/VulkanRHI/Private/VulkanTexture.cpp
--- a/Engine/Source/Runtime/VulkanRHI/Private/VulkanTexture.cpp
+++ b/Engine/Source/Runtime/VulkanRHI/Private/VulkanTexture.cpp
@@ -131,7 +131,7 @@ static void VulkanTextureDestroyed(const FRHITextureDesc& TextureDesc, uint64 Si
 	UpdateVulkanTextureStats(TextureDesc, Size, false);
 }
 
-void FVulkanTexture::InternalLockWrite(FVulkanContextCommon& Context, FVulkanTexture* Surface, const VkBufferImageCopy& Region, VulkanRHI::FStagingBuffer* StagingBuffer)
+void FVulkanTexture::InternalLockWrite(FVulkanContextCommon& Context, FVulkanTexture* Surface, const VkBufferImageCopy& Region, VulkanRHI::FStagingBuffer* StagingBuffer, bool bNeedsTransitions)
 {
 	FVulkanCommandBuffer* CmdBuffer = Context.GetActiveCmdBuffer();
 	ensure(CmdBuffer->IsOutsideRenderPass());
@@ -140,6 +140,8 @@ void FVulkanTexture::InternalLockWrite(FVulkanContextCommon& Context, FVulkanTex
 	const VkImageSubresourceLayers& ImageSubresource = Region.imageSubresource;
 	const VkImageSubresourceRange SubresourceRange = FVulkanPipelineBarrier::MakeSubresourceRange(ImageSubresource.aspectMask, ImageSubresource.mipLevel, 1, ImageSubresource.baseArrayLayer, ImageSubresource.layerCount);
 
+	// We might skip transitions in passes where copies are batched
+	if (bNeedsTransitions)
 	{
 		FVulkanPipelineBarrier Barrier;
 		Barrier.AddImageLayoutTransition(Surface->Image, VK_IMAGE_LAYOUT_UNDEFINED, VK_IMAGE_LAYOUT_TRANSFER_DST_OPTIMAL, SubresourceRange);
@@ -148,8 +150,9 @@ void FVulkanTexture::InternalLockWrite(FVulkanContextCommon& Context, FVulkanTex
 
 	VulkanRHI::vkCmdCopyBufferToImage(StagingCommandBuffer, StagingBuffer->GetHandle(), Surface->Image, VK_IMAGE_LAYOUT_TRANSFER_DST_OPTIMAL, 1, &Region);
 
-	// :todo-jn: replace with cmdlist layout tracking (ideally would happen on UploadContext)
+	if (bNeedsTransitions)
 	{
+		// :todo-jn: replace with cmdlist layout tracking
 		FVulkanPipelineBarrier Barrier;
 		Barrier.AddImageLayoutTransition(Surface->Image, VK_IMAGE_LAYOUT_TRANSFER_DST_OPTIMAL, Surface->GetDefaultLayout(), SubresourceRange);
 		Barrier.Execute(CmdBuffer);
@@ -1621,10 +1624,11 @@ void FVulkanDynamicRHI::RHIUnlockTexture(FRHICommandListImmediate& RHICmdList, c
 	Region.imageExtent.height = MipHeight;
 	Region.imageExtent.depth = 1;
 
+	const bool bNeedsTransition = RHICmdList.NeedsExtraTransitions();
 	RHICmdList.EnqueueLambda(TEXT("FVulkanTexture::InternalLockWrite"),
-		[Texture, Region, StagingBuffer](FRHICommandListBase& ExecutingCmdList)
+		[Texture, Region, StagingBuffer, bNeedsTransition](FRHICommandListBase& ExecutingCmdList)
 		{
-			FVulkanTexture::InternalLockWrite(FVulkanCommandListContext::Get(ExecutingCmdList), Texture, Region, StagingBuffer);
+			FVulkanTexture::InternalLockWrite(FVulkanCommandListContext::Get(ExecutingCmdList), Texture, Region, StagingBuffer, bNeedsTransition);
 		});
 }
 
@@ -1676,10 +1680,11 @@ void FVulkanDynamicRHI::InternalUpdateTexture2D(FRHICommandListBase& RHICmdList,
 	Region.imageExtent.depth = 1;
 
 	FVulkanTexture* Texture = ResourceCast(TextureRHI);
+	const bool bNeedsTransition = RHICmdList.NeedsExtraTransitions();
 	RHICmdList.EnqueueLambda(TEXT("FVulkanTexture::InternalLockWrite"),
-		[Texture, Region, StagingBuffer](FRHICommandListBase& ExecutingCmdList)
+		[Texture, Region, StagingBuffer, bNeedsTransition](FRHICommandListBase& ExecutingCmdList)
 		{
-			FVulkanTexture::InternalLockWrite(FVulkanCommandListContext::Get(ExecutingCmdList), Texture, Region, StagingBuffer);
+			FVulkanTexture::InternalLockWrite(FVulkanCommandListContext::Get(ExecutingCmdList), Texture, Region, StagingBuffer, bNeedsTransition);
 		});
 }
 
@@ -1770,10 +1775,11 @@ void FVulkanDynamicRHI::InternalUpdateTexture3D(FRHICommandListBase& RHICmdList,
 	Region.imageExtent.height = (uint32)FMath::Min((int32)(TextureSizeY-UpdateRegion.DestY), (int32)UpdateRegion.Height);
 	Region.imageExtent.depth = (uint32)FMath::Min((int32)(TextureSizeZ-UpdateRegion.DestZ), (int32)UpdateRegion.Depth);
 
+	const bool bNeedsTransition = RHICmdList.NeedsExtraTransitions();
 	RHICmdList.EnqueueLambda(TEXT("FVulkanTexture::InternalLockWrite"),
-		[Texture, Region, StagingBuffer](FRHICommandListBase& ExecutingCmdList)
+		[Texture, Region, StagingBuffer, bNeedsTransition](FRHICommandListBase& ExecutingCmdList)
 		{
-			FVulkanTexture::InternalLockWrite(FVulkanCommandListContext::Get(ExecutingCmdList), Texture, Region, StagingBuffer);
+			FVulkanTexture::InternalLockWrite(FVulkanCommandListContext::Get(ExecutingCmdList), Texture, Region, StagingBuffer, bNeedsTransition);
 		});
 }
 
