VulkanRHI: Remove redundant BindPipeline calls at the command buffer level (WIP).

#rnx

#Synced-CL 49847350

diff --git a/Engine/Source/Runtime/VulkanRHI/Private/VulkanCommandBuffer.cpp b/Engine/Source/Runtime/VulkanRHI/Private/VulkanCommandBuffer.cpp
--- a/Engine/Source/Runtime/VulkanRHI/Private/VulkanCommandBuffer.cpp
+++ b/Engine/Source/Runtime/VulkanRHI/Private/VulkanCommandBuffer.cpp
@@ -55,6 +55,10 @@ void FVulkanCommandBuffer::AllocMemory()
 	bHasStencilRef = 0;
 	State = EState::ReadyForBegin;
 
+	CurrentPipelineHandle = VK_NULL_HANDLE;
+	CurrentComputeShaderObject = VK_NULL_HANDLE;
+	FMemory::Memzero(CurrentGraphicsShaderObjects);
+
 	INC_DWORD_STAT(STAT_VulkanNumCmdBuffers);
 }
 
@@ -265,6 +269,10 @@ void FVulkanCommandBuffer::Reset()
 		CurrentScissors.Empty();
 		CurrentStencilRef = 0;
 
+		CurrentPipelineHandle = VK_NULL_HANDLE;
+		CurrentComputeShaderObject = VK_NULL_HANDLE;
+		FMemory::Memzero(CurrentGraphicsShaderObjects);
+
 		for (VkEvent BarrierEvent : EndedBarrierEvents)
 		{
 			Device.ReleaseBarrierEvent(BarrierEvent);
@@ -276,6 +284,84 @@ void FVulkanCommandBuffer::Reset()
 	}
 }
 
+void FVulkanCommandBuffer::BindPipeline(VkPipelineBindPoint BindPoint, VkPipeline PipelineHandle)
+{
+	if (PipelineHandle != CurrentPipelineHandle)
+	{
+		VulkanRHI::vkCmdBindPipeline(CommandBufferHandle, BindPoint, PipelineHandle);
+		CurrentPipelineHandle = PipelineHandle;
+	}
+}
+
+void FVulkanCommandBuffer::BindGraphicsShaderObjects(TArrayView<FVulkanShader*> VulkanShaders)
+{
+	check(VulkanShaders.Num() == ShaderStage::NumGraphicsStages);
+	VkShaderStageFlagBits ShaderStages[ShaderStage::NumGraphicsStages];
+	VkShaderEXT ShaderObjectHandles[ShaderStage::NumGraphicsStages] = { VK_NULL_HANDLE };
+
+	uint32 ChangeCount = 0;
+	auto AddStage = [this, &ChangeCount, &ShaderStages, &ShaderObjectHandles, &VulkanShaders](ShaderStage::EStage Stage, VkShaderStageFlagBits StageFlag)
+	{
+
+		FVulkanShader* VulkanShader = VulkanShaders[Stage];
+		if (VulkanShader)
+		{
+			VulkanShader->WaitForAsyncCompile();
+			ShaderObjectHandles[ChangeCount] = VulkanShader->GetShaderObject();
+		}
+		else
+		{
+			ShaderObjectHandles[ChangeCount] = VK_NULL_HANDLE;
+		}
+
+		if (CurrentGraphicsShaderObjects[Stage] != ShaderObjectHandles[ChangeCount])
+		{
+			ShaderStages[ChangeCount] = StageFlag;
+			CurrentGraphicsShaderObjects[Stage] = ShaderObjectHandles[ChangeCount];
+			ChangeCount++;
+		}
+	};
+
+	AddStage(ShaderStage::Vertex, VK_SHADER_STAGE_VERTEX_BIT);
+	AddStage(ShaderStage::Pixel, VK_SHADER_STAGE_FRAGMENT_BIT);
+
+#if PLATFORM_SUPPORTS_GEOMETRY_SHADERS
+	if (Device.GetPhysicalDeviceFeatures().Core_1_0.geometryShader)
+	{
+		AddStage(ShaderStage::Geometry, VK_SHADER_STAGE_GEOMETRY_BIT);
+	}
+#endif
+
+#if PLATFORM_SUPPORTS_MESH_SHADERS
+	if (Device.GetOptionalExtensions().HasEXTMeshShader)
+	{
+		AddStage(ShaderStage::Mesh, VK_SHADER_STAGE_MESH_BIT_EXT);
+		AddStage(ShaderStage::Task, VK_SHADER_STAGE_TASK_BIT_EXT);
+	}
+#endif
+
+	if (ChangeCount > 0)
+	{
+		VulkanRHI::vkCmdBindShadersEXT(CommandBufferHandle, ChangeCount, ShaderStages, ShaderObjectHandles);
+	}
+}
+
+void FVulkanCommandBuffer::BindComputeShaderObject(FVulkanShader* VulkanShader)
+{
+	check(VulkanShader);
+	VulkanShader->WaitForAsyncCompile();
+
+	static const VkShaderStageFlagBits ShaderStage = VK_SHADER_STAGE_COMPUTE_BIT;
+	const VkShaderEXT ShaderObjectHandle = VulkanShader->GetShaderObject();
+	check(ShaderObjectHandle);
+
+	if (CurrentComputeShaderObject != ShaderObjectHandle)
+	{
+		VulkanRHI::vkCmdBindShadersEXT(CommandBufferHandle, 1, &ShaderStage, &ShaderObjectHandle);
+		CurrentComputeShaderObject = ShaderObjectHandle;
+	}
+}
+
 void FVulkanCommandBuffer::SetSubmitted()
 {
 	for (FVulkanCommandBuffer* SecondaryCommand : ExecutedSecondaryCommandBuffers)
diff --git a/Engine/Source/Runtime/VulkanRHI/Private/VulkanCommandBuffer.h b/Engine/Source/Runtime/VulkanRHI/Private/VulkanCommandBuffer.h
--- a/Engine/Source/Runtime/VulkanRHI/Private/VulkanCommandBuffer.h
+++ b/Engine/Source/Runtime/VulkanRHI/Private/VulkanCommandBuffer.h
@@ -15,6 +15,7 @@ class FVulkanCommandBufferPool;
 class FVulkanContextCommon;
 class FVulkanQueue;
 class FVulkanQueryPool;
+class FVulkanShader;
 struct FVulkanBeginRenderPassInfo;
 struct FGfxPipelineDesc;
 
@@ -166,6 +167,10 @@ public:
 
 	void SetUnusedGraphicsDynamicStates();
 
+	void BindPipeline(VkPipelineBindPoint BindPoint, VkPipeline PipelineHandle);
+	void BindGraphicsShaderObjects(TArrayView<FVulkanShader*> VulkanShaders);
+	void BindComputeShaderObject(FVulkanShader* VulkanShader);
+
 private:
 	VkCommandBuffer CommandBufferHandle = VK_NULL_HANDLE;
 	double SubmittedTime = 0.0f;
@@ -174,6 +179,10 @@ private:
 
 	FVulkanCommandBufferPool& CommandBufferPool;
 
+	VkPipeline CurrentPipelineHandle = VK_NULL_HANDLE;
+	VkShaderEXT CurrentComputeShaderObject = VK_NULL_HANDLE;
+	VkShaderEXT CurrentGraphicsShaderObjects[ShaderStage::NumGraphicsStages] = { VK_NULL_HANDLE };
+
 	TArray<VkEvent> EndedBarrierEvents;
 
 	void AllocMemory();
diff --git a/Engine/Source/Runtime/VulkanRHI/Private/VulkanCommands.cpp b/Engine/Source/Runtime/VulkanRHI/Private/VulkanCommands.cpp
--- a/Engine/Source/Runtime/VulkanRHI/Private/VulkanCommands.cpp
+++ b/Engine/Source/Runtime/VulkanRHI/Private/VulkanCommands.cpp
@@ -614,7 +614,7 @@ void FVulkanCommandListContext::RHISetGraphicsPipelineState(FRHIGraphicsPipeline
 #if VULKAN_ENABLE_AGGRESSIVE_STATS
 		SCOPE_CYCLE_COUNTER(STAT_VulkanPipelineBind);
 #endif
-		PendingGfxState->Bind(CommandBuffer.GetHandle());
+		PendingGfxState->Bind(CommandBuffer);
 		CommandBuffer.bHasPipeline = true;
 	}
 
diff --git a/Engine/Source/Runtime/VulkanRHI/Private/VulkanPendingState.cpp b/Engine/Source/Runtime/VulkanRHI/Private/VulkanPendingState.cpp
--- a/Engine/Source/Runtime/VulkanRHI/Private/VulkanPendingState.cpp
+++ b/Engine/Source/Runtime/VulkanRHI/Private/VulkanPendingState.cpp
@@ -340,7 +340,7 @@ void FVulkanPendingComputeState::PrepareForDispatch(FVulkanCommandListContext& C
 	if (CurrentPipeline->UsesBindless())
 	{
 		CurrentState->UpdateBindlessDescriptors(Context);
-		CurrentPipeline->Bind(Context.GetCommandBuffer().GetHandle());
+		CurrentPipeline->Bind(Context.GetCommandBuffer());
 	}
 	else
 	{
@@ -351,11 +351,10 @@ void FVulkanPendingComputeState::PrepareForDispatch(FVulkanCommandListContext& C
 		SCOPE_CYCLE_COUNTER(STAT_VulkanPipelineBind);
 #endif
 
-		const VkCommandBuffer CommandBufferHandle = Context.GetCommandBuffer().GetHandle();
-		CurrentPipeline->Bind(CommandBufferHandle);
+		CurrentPipeline->Bind(Context.GetCommandBuffer());
 		if (bHasDescriptorSets)
 		{
-			CurrentState->BindDescriptorSets(CommandBufferHandle);
+			CurrentState->BindDescriptorSets(Context.GetCommandBuffer().GetHandle());
 		}
 	}
 }
diff --git a/Engine/Source/Runtime/VulkanRHI/Private/VulkanPendingState.h b/Engine/Source/Runtime/VulkanRHI/Private/VulkanPendingState.h
--- a/Engine/Source/Runtime/VulkanRHI/Private/VulkanPendingState.h
+++ b/Engine/Source/Runtime/VulkanRHI/Private/VulkanPendingState.h
@@ -208,9 +208,9 @@ public:
 		bDirtyVertexStreams = true;
 	}
 
-	void Bind(VkCommandBuffer CmdBuffer)
+	void Bind(FVulkanCommandBuffer& CommandBuffer)
 	{
-		CurrentPipeline->Bind(CmdBuffer);
+		CurrentPipeline->Bind(CommandBuffer);
 	}
 
 	FVulkanPackedUniformBufferState& GetUniformBuffers(uint8 Stage)
diff --git a/Engine/Source/Runtime/VulkanRHI/Private/VulkanPipeline.cpp b/Engine/Source/Runtime/VulkanRHI/Private/VulkanPipeline.cpp
--- a/Engine/Source/Runtime/VulkanRHI/Private/VulkanPipeline.cpp
+++ b/Engine/Source/Runtime/VulkanRHI/Private/VulkanPipeline.cpp
@@ -3179,28 +3179,29 @@ void FVulkanGraphicsPipelineState::DeleteVkPipeline(bool bImmediate)
 	Device->PipelineStateCache->LRUCheckNotInside(this);
 }
 
-void FVulkanGraphicsPipelineState::Bind(VkCommandBuffer CmdBuffer)
+void FVulkanGraphicsPipelineState::Bind(FVulkanCommandBuffer& CommandBuffer)
 {
 	if (Device->SupportsShaderObjects())
 	{
-		Device->GetShaderObjectManager()->BindGraphicsShaderObjects(CmdBuffer, VulkanShaders);
+		CommandBuffer.BindGraphicsShaderObjects(VulkanShaders);
 	}
 	else
 	{
-		VulkanRHI::vkCmdBindPipeline(CmdBuffer, VK_PIPELINE_BIND_POINT_GRAPHICS, VulkanPipeline);
+		CommandBuffer.BindPipeline(VK_PIPELINE_BIND_POINT_GRAPHICS, VulkanPipeline);
 	}
 }
 
-void FVulkanComputePipeline::Bind(VkCommandBuffer CmdBuffer)
+void FVulkanComputePipeline::Bind(FVulkanCommandBuffer& CommandBuffer)
 {
 	if (Device->SupportsShaderObjects())
 	{
 		check(ComputeShader.IsValid());
-		Device->GetShaderObjectManager()->BindComputeShaderObject(CmdBuffer, ResourceCast(GetComputeShader()));
+		FVulkanShader* VulkanShader = ResourceCast(GetComputeShader());
+		CommandBuffer.BindComputeShaderObject(VulkanShader);
 	}
 	else
 	{
-		VulkanRHI::vkCmdBindPipeline(CmdBuffer, VK_PIPELINE_BIND_POINT_COMPUTE, Pipeline);
+		CommandBuffer.BindPipeline(VK_PIPELINE_BIND_POINT_COMPUTE, Pipeline);
 	}
 }
 
diff --git a/Engine/Source/Runtime/VulkanRHI/Private/VulkanPipeline.h b/Engine/Source/Runtime/VulkanRHI/Private/VulkanPipeline.h
--- a/Engine/Source/Runtime/VulkanRHI/Private/VulkanPipeline.h
+++ b/Engine/Source/Runtime/VulkanRHI/Private/VulkanPipeline.h
@@ -717,7 +717,7 @@ public:
 		return *Layout;
 	}
 
-	void Bind(VkCommandBuffer CmdBuffer);
+	void Bind(FVulkanCommandBuffer& CommandBuffer);
 
 	bool UsesBindless() const
 	{
@@ -744,7 +744,7 @@ public:
 		return *Layout;
 	}
 
-	void Bind(VkCommandBuffer CmdBuffer);
+	void Bind(FVulkanCommandBuffer& CommandBuffer);
 
 	const uint64 GetShaderKey(EShaderFrequency Frequency) const
 	{
diff --git a/Engine/Source/Runtime/VulkanRHI/Private/VulkanShaderObjectManager.cpp b/Engine/Source/Runtime/VulkanRHI/Private/VulkanShaderObjectManager.cpp
--- a/Engine/Source/Runtime/VulkanRHI/Private/VulkanShaderObjectManager.cpp
+++ b/Engine/Source/Runtime/VulkanRHI/Private/VulkanShaderObjectManager.cpp
@@ -130,71 +130,3 @@ bool FVulkanShaderObjectManager::CreateLinkedShaderObjects(FVulkanGraphicsPipeli
 	return true;
 }
 
-void FVulkanShaderObjectManager::BindGraphicsShaderObjects(VkCommandBuffer CommandBufferHandle, TArrayView<FVulkanShader*> VulkanShaders)
-{
-	VkShaderStageFlagBits ShaderStages[ShaderStage::NumGraphicsStages];
-	VkShaderEXT ShaderObjectHandles[ShaderStage::NumGraphicsStages] = { VK_NULL_HANDLE };
-
-	uint32 StageCount = 0;
-	auto AddStage = [&StageCount, &ShaderStages, &ShaderObjectHandles, &VulkanShaders](ShaderStage::EStage Stage, VkShaderStageFlagBits StageFlag)
-	{
-		ShaderStages[StageCount] = StageFlag;
-
-		FVulkanShader* VulkanShader = VulkanShaders[Stage];
-		if (VulkanShader)
-		{
-			if (VulkanShader->AsyncCompileTask.IsValid() && !VulkanShader->AsyncCompileTask->IsComplete())
-			{
-				VulkanShader->AsyncCompileTask->Wait();
-			}
-			ShaderObjectHandles[StageCount] = VulkanShader->ShaderObject;
-			checkf(Stage == ShaderStage::GetStageForFrequency(VulkanShader->Frequency), 
-				TEXT("Shader Frequency (%s) does not match expected Vulkan shader stage (%d)!"), 
-				GetShaderFrequencyString(VulkanShader->Frequency), Stage);
-		}
-		else
-		{
-			ShaderObjectHandles[StageCount] = VK_NULL_HANDLE;
-		}
-
-		StageCount++;
-	};
-
-	AddStage(ShaderStage::Vertex, VK_SHADER_STAGE_VERTEX_BIT);
-	AddStage(ShaderStage::Pixel, VK_SHADER_STAGE_FRAGMENT_BIT);
-
-#if PLATFORM_SUPPORTS_GEOMETRY_SHADERS
-	if (Device.GetPhysicalDeviceFeatures().Core_1_0.geometryShader)
-	{
-		AddStage(ShaderStage::Geometry, VK_SHADER_STAGE_GEOMETRY_BIT);
-	}
-#endif
-
-#if PLATFORM_SUPPORTS_MESH_SHADERS
-	if (Device.GetOptionalExtensions().HasEXTMeshShader)
-	{
-		AddStage(ShaderStage::Mesh, VK_SHADER_STAGE_MESH_BIT_EXT);
-		AddStage(ShaderStage::Task, VK_SHADER_STAGE_TASK_BIT_EXT);
-	}
-#endif
-
-	// :todo-jn: Cache state and avoid setting the same shader over and over
-	VulkanRHI::vkCmdBindShadersEXT(CommandBufferHandle, StageCount, ShaderStages, ShaderObjectHandles);
-}
-
-void FVulkanShaderObjectManager::BindComputeShaderObject(VkCommandBuffer CommandBufferHandle, FVulkanShader* VulkanShader)
-{
-	check(VulkanShader);
-
-	if (VulkanShader->AsyncCompileTask.IsValid() && !VulkanShader->AsyncCompileTask->IsComplete())
-	{
-		VulkanShader->AsyncCompileTask->Wait();
-	}
-
-	check(VulkanShader->ShaderObject);
-
-	static const VkShaderStageFlagBits ShaderStage = VK_SHADER_STAGE_COMPUTE_BIT;
-	const VkShaderEXT ShaderObjectHandle = VulkanShader->ShaderObject;
-
-	VulkanRHI::vkCmdBindShadersEXT(CommandBufferHandle, 1, &ShaderStage, &ShaderObjectHandle);
-}
diff --git a/Engine/Source/Runtime/VulkanRHI/Private/VulkanShaderObjectManager.h b/Engine/Source/Runtime/VulkanRHI/Private/VulkanShaderObjectManager.h
--- a/Engine/Source/Runtime/VulkanRHI/Private/VulkanShaderObjectManager.h
+++ b/Engine/Source/Runtime/VulkanRHI/Private/VulkanShaderObjectManager.h
@@ -19,9 +19,6 @@ public:
 	void CreateShaderObject(FVulkanShader& VulkanShader);
 	bool CreateLinkedShaderObjects(FVulkanGraphicsPipelineState* PSO, FVulkanShader* Shaders[ShaderStage::NumGraphicsStages], FGraphicsPipelineStateInitializer::EPSOPrecacheCompileType PSOCompileType);
 
-	void BindGraphicsShaderObjects(VkCommandBuffer CommandBufferHandle, TArrayView<FVulkanShader*> VulkanShaders);
-	void BindComputeShaderObject(VkCommandBuffer CommandBufferHandle, FVulkanShader* VulkanShader);
-
 private:
 	VkShaderStageFlags GetPossibleNextStages(VkShaderStageFlagBits Stage);
 	void InternalCreateShaderObject(FVulkanShader& VulkanShader);
