VulkanRHI: Reduce default size of shader record

#rnx

#Synced-CL 49847350

diff --git a/Engine/Source/Runtime/VulkanRHI/Private/VulkanRayTracing.cpp b/Engine/Source/Runtime/VulkanRHI/Private/VulkanRayTracing.cpp
--- a/Engine/Source/Runtime/VulkanRHI/Private/VulkanRayTracing.cpp
+++ b/Engine/Source/Runtime/VulkanRHI/Private/VulkanRayTracing.cpp
@@ -34,11 +34,11 @@ static FAutoConsoleVariableRef CVarVulkanRayTracingCompactionMinPrimitiveCount(
 	TEXT("Sets the minimum primitive count threshold below which geometry skips the compaction. (default = 128)\n")
 );
 
-static int32 GVulkanRayTracingMaxShaderGroupStride = 4096;
+static int32 GVulkanRayTracingMaxShaderGroupStride = 512;
 static FAutoConsoleVariableRef GCVarVulkanRayTracingMaxShaderGroupStride(
 	TEXT("r.Vulkan.RayTracing.MaxShaderGroupStride"),
 	GVulkanRayTracingMaxShaderGroupStride,
-	TEXT("The default size to allocate for each record (default: 4096)."),
+	TEXT("The default size to allocate for each record (default: 512)."),
 	ECVF_ReadOnly
 );
 
