VulkanRHI: Add maintenance9 as a requirement for multi-queue.

#rnx

#Synced-CL 49847350

diff --git a/Engine/Source/Runtime/VulkanRHI/Private/VulkanDevice.cpp b/Engine/Source/Runtime/VulkanRHI/Private/VulkanDevice.cpp
--- a/Engine/Source/Runtime/VulkanRHI/Private/VulkanDevice.cpp
+++ b/Engine/Source/Runtime/VulkanRHI/Private/VulkanDevice.cpp
@@ -371,10 +371,12 @@ void FVulkanDevice::CreateDevice(TArray<const ANSICHAR*>& DeviceLayers, FVulkanD
 			// - cvars allow for a dedicated async compute queue
 			// - a new family index is available
 			// - Sync2 is available
+			// - maintenance9 is available
 			if ((ComputeQueueFamilyIndex == -1) &&
 				(GRHIAllowAsyncComputeCvar.GetValueOnAnyThread() != 0 || GAllowPresentOnComputeQueue.GetValueOnAnyThread() != 0 || FParse::Param(FCommandLine::Get(), TEXT("ForceAsyncCompute"))) &&
 				(GfxQueueFamilyIndex != FamilyIndex) && 
-				SupportsParallelRendering())
+				SupportsParallelRendering() &&
+				OptionalDeviceExtensions.HasKHRMaintenance9)
 			{
 				ComputeQueueFamilyIndex = FamilyIndex;
 				bIsValidQueue = true;
@@ -383,11 +385,19 @@ void FVulkanDevice::CreateDevice(TArray<const ANSICHAR*>& DeviceLayers, FVulkanD
 		
 		if (VKHasAllFlags(CurrProps.queueFlags, VK_QUEUE_TRANSFER_BIT))
 		{
-			// Prefer a non-gfx transfer queue
+			// Allocate a transfer queue if:
+			// - a transfer queue hasn't been found already
+			// - there is a dedicated non graphics/compute queue
+			// - cvars allow for a dedicated transfer queue
+			// - a new family index is available
+			// - Sync2 is available
+			// - maintenance9 is available
 			if ((TransferQueueFamilyIndex == -1) &&
 				!VKHasAnyFlags(CurrProps.queueFlags, VK_QUEUE_GRAPHICS_BIT | VK_QUEUE_COMPUTE_BIT) &&
 				(GRHIAllowTransferQueueCvar.GetValueOnAnyThread() != 0) &&
-				SupportsParallelRendering())
+				(GfxQueueFamilyIndex != FamilyIndex) && (ComputeQueueFamilyIndex != FamilyIndex) &&
+				SupportsParallelRendering() &&
+				OptionalDeviceExtensions.HasKHRMaintenance9)
 			{
 				TransferQueueFamilyIndex = FamilyIndex;
 				bIsValidQueue = true;
