Fix printf-style format specifier/type mismatches in VulkanDebug.

#Synced-CL 49904282

diff --git a/Engine/Source/Runtime/VulkanRHI/Private/VulkanDebug.cpp b/Engine/Source/Runtime/VulkanRHI/Private/VulkanDebug.cpp
--- a/Engine/Source/Runtime/VulkanRHI/Private/VulkanDebug.cpp
+++ b/Engine/Source/Runtime/VulkanRHI/Private/VulkanDebug.cpp
@@ -1355,7 +1355,7 @@ void FWrapLayer::WaitForFences(VkResult Result, VkDevice Device, uint32 FenceCou
 #if VULKAN_ENABLE_DUMP_LAYER
 	if (Result == VK_RESULT_MAX_ENUM)
 	{
-		DevicePrintfBegin(Device, FString::Printf(TEXT("vkWaitForFences(Count=0x%p, Fences=%d, WaitAll=%d, Timeout=0x%p)"), FenceCount, Fences, (uint32)bWaitAll, Timeout));
+		DevicePrintfBegin(Device, FString::Printf(TEXT("vkWaitForFences(Count=%d, Fences=0x%p, WaitAll=%d, Timeout=0x%llx)"), FenceCount, Fences, (uint32)bWaitAll, Timeout));
 		DumpFenceList(FenceCount, Fences);
 	}
 	else
@@ -3308,7 +3308,7 @@ void FWrapLayer::AcquireNextImageKHR(VkResult Result, VkDevice Device, VkSwapcha
 	if (Result == VK_RESULT_MAX_ENUM)
 	{
 #if VULKAN_ENABLE_DUMP_LAYER
-		DevicePrintfBeginResult(Device, FString::Printf(TEXT("vkAcquireNextImageKHR(Swapchain=0x%p, Timeout=0x%p, Semaphore=0x%p, Fence=0x%p, OutImageIndex=0x%p)[...]\n"), Swapchain, (void*)Timeout, Semaphore, Fence, ImageIndex));
+		DevicePrintfBeginResult(Device, FString::Printf(TEXT("vkAcquireNextImageKHR(Swapchain=0x%p, Timeout=0x%llx, Semaphore=0x%p, Fence=0x%p, OutImageIndex=0x%p)[...]\n"), Swapchain, Timeout, Semaphore, Fence, ImageIndex));
 #endif
 	}
 	else
@@ -4639,7 +4639,7 @@ void FWrapLayer::CmdResetEvent2KHR(VkResult Result, VkCommandBuffer CommandBuffe
 	if (Result == VK_RESULT_MAX_ENUM)
 	{
 #if VULKAN_ENABLE_DUMP_LAYER
-		PrintfBeginResult(FString::Printf(TEXT("CmdResetEvent2KHR(CommandBuffer=0x%p, Event=0x%p, StageMask=0xllx)"), CommandBuffer, Event, StageMask));
+		PrintfBeginResult(FString::Printf(TEXT("CmdResetEvent2KHR(CommandBuffer=0x%p, Event=0x%p, StageMask=0x%llx)"), CommandBuffer, Event, StageMask));
 #endif
 	}
 }
@@ -4721,7 +4721,7 @@ void FWrapLayer::WaitSemaphoresKHR(VkResult Result, VkDevice Device, const VkSem
 	if (Result == VK_RESULT_MAX_ENUM)
 	{
 #if VULKAN_ENABLE_DUMP_LAYER
-		PrintfBeginResult(FString::Printf(TEXT("WaitSemaphoresKHR(Device=0x%p, WaitInfo=0x%p, Timeout=0xllx))"), Device, WaitInfo, Timeout));
+		PrintfBeginResult(FString::Printf(TEXT("WaitSemaphoresKHR(Device=0x%p, WaitInfo=0x%p, Timeout=0x%llx))"), Device, WaitInfo, Timeout));
 #endif
 	}
 }
