[Vulkan RHI] Remove redundant frame-pacing implementation.

#rn Removed r.Vulkan.CPURenderthreadFramePacer which resulted in artificially stunted frame-rate in ASIS multi-view.
#Synced-CL 49972327

diff --git a/Engine/Source/Runtime/VulkanRHI/Private/Android/VulkanAndroidPlatform.cpp b/Engine/Source/Runtime/VulkanRHI/Private/Android/VulkanAndroidPlatform.cpp
--- a/Engine/Source/Runtime/VulkanRHI/Private/Android/VulkanAndroidPlatform.cpp
+++ b/Engine/Source/Runtime/VulkanRHI/Private/Android/VulkanAndroidPlatform.cpp
@@ -58,7 +58,6 @@ namespace AndroidVulkan
 #endif // #if USE_ANDROID_VULKAN_SWAPPY
 
 // From VulklanSwapChain.cpp
-extern int32 GVulkanCPURenderThreadFramePacer;
 extern int32 GPrintVulkanVsyncDebug;
 
 int32 GVulkanExtensionFramePacer = 1;
@@ -978,7 +977,6 @@ VkResult FVulkanAndroidPlatform::CreateSwapchainKHR(FVulkanPlatformWindowContext
 				UE_LOG(LogVulkanRHI, Log, TEXT("SwappyVk_initAndGetRefreshCycleDuration: %ull"), RefreshDuration);
 			}
 
-			GVulkanCPURenderThreadFramePacer = 0;
 			GVulkanExtensionFramePacer = 0;
 		}
 		else
@@ -987,7 +985,6 @@ VkResult FVulkanAndroidPlatform::CreateSwapchainKHR(FVulkanPlatformWindowContext
 		if (GVulkanExtensionFramePacer && FVulkanAndroidPlatform::bHasGoogleDisplayTiming)
 		{
 			GDTimingFramePacer = MakeUnique<FGDTimingFramePacer>(Device, *Swapchain);
-			GVulkanCPURenderThreadFramePacer = 0;
 		}
 #else
 		{}
diff --git a/Engine/Source/Runtime/VulkanRHI/Private/VulkanSwapChain.cpp b/Engine/Source/Runtime/VulkanRHI/Private/VulkanSwapChain.cpp
--- a/Engine/Source/Runtime/VulkanRHI/Private/VulkanSwapChain.cpp
+++ b/Engine/Source/Runtime/VulkanRHI/Private/VulkanSwapChain.cpp
@@ -46,15 +46,6 @@ static FAutoConsoleVariableRef CVarCpuWaitForFence(
 	ECVF_RenderThreadSafe
 );
 
-//disabled by default in swapchain creation if the extension frame pacer is available
-int32 GVulkanCPURenderThreadFramePacer = 1;
-static FAutoConsoleVariableRef CVarVulkanCPURenderThreadFramePacer(
-	TEXT("r.Vulkan.CPURenderthreadFramePacer"),
-	GVulkanCPURenderThreadFramePacer,
-	TEXT("Whether to enable the simple Render thread CPU Framepacer for Vulkan"),
-	ECVF_RenderThreadSafe
-);
-
 int32 GVulkanCPURHIFramePacer = 1;
 static FAutoConsoleVariableRef CVarVulkanCPURHIFramePacer(
 	TEXT("r.Vulkan.CPURHIThreadFramePacer"),
@@ -771,53 +762,6 @@ int32 FVulkanSwapChain::AcquireImageIndex(FVulkanSemaphore** OutSemaphore)
 	return CurrentImageIndex;
 }
 
-void FVulkanSwapChain::RenderThreadPacing()
-{
-	check(IsInRenderingThread());
-	const int32 SyncInterval = (LockToVsync || GVulkanForcePacingWithoutVSync) ? RHIGetSyncInterval() : 0;
-
-	//very naive CPU side frame pacer.
-	if (GVulkanCPURenderThreadFramePacer && SyncInterval > 0)
-	{
-		double NowCPUTime = FPlatformTime::Seconds();
-		double DeltaCPUPresentTimeMS = (NowCPUTime - RTPacingPreviousFrameCPUTime) * 1000.0;
-
-
-		double TargetIntervalWithEpsilonMS = (double)SyncInterval * (1.0 / 60.0) * 1000.0;
-		const double IntervalThresholdMS = TargetIntervalWithEpsilonMS * 0.1;
-
-		RTPacingSampledDeltaTimeMS += DeltaCPUPresentTimeMS; RTPacingSampleCount++;
-
-		double SampledDeltaMS = (RTPacingSampledDeltaTimeMS / (double)RTPacingSampleCount) + IntervalThresholdMS;
-
-		if (RTPacingSampleCount > 1000)
-		{
-			RTPacingSampledDeltaTimeMS = SampledDeltaMS;
-			RTPacingSampleCount = 1;
-		}
-
-		if (SampledDeltaMS < (TargetIntervalWithEpsilonMS))
-		{
-			FRenderThreadIdleScope IdleScope(ERenderThreadIdleTypes::WaitingForGPUPresent);
-			QUICK_SCOPE_CYCLE_COUNTER(STAT_StallForEmulatedSyncInterval);
-
-			FPlatformProcess::SleepNoStats((TargetIntervalWithEpsilonMS - SampledDeltaMS) * 0.001f);
-			if (GPrintVulkanVsyncDebug)
-			{
-				UE_LOG(LogVulkanRHI, Log, TEXT("CPU RT delta: %f, TargetWEps: %f, sleepTime: %f "), SampledDeltaMS, TargetIntervalWithEpsilonMS, TargetIntervalWithEpsilonMS - DeltaCPUPresentTimeMS);
-			}
-		}
-		else
-		{
-			if (GPrintVulkanVsyncDebug)
-			{
-				UE_LOG(LogVulkanRHI, Log, TEXT("CPU RT delta: %f"), DeltaCPUPresentTimeMS);
-			}
-		}
-		RTPacingPreviousFrameCPUTime = NowCPUTime;
-	}
-}
-
 FVulkanSwapChain::EStatus FVulkanSwapChain::Present(FVulkanQueue* PresentQueue, FVulkanSemaphore* BackBufferRenderingDoneSemaphore)
 {
 	check(CurrentImageIndex != -1);
diff --git a/Engine/Source/Runtime/VulkanRHI/Private/VulkanSwapChain.h b/Engine/Source/Runtime/VulkanRHI/Private/VulkanSwapChain.h
--- a/Engine/Source/Runtime/VulkanRHI/Private/VulkanSwapChain.h
+++ b/Engine/Source/Runtime/VulkanRHI/Private/VulkanSwapChain.h
@@ -50,7 +50,6 @@ public:
 	};
 	EStatus Present(FVulkanQueue* PresentQueue, FVulkanSemaphore* BackBufferRenderingDoneSemaphore);
 
-	void RenderThreadPacing();
 	inline int8 DoesLockToVsync() { return LockToVsync; }
 
 	inline VkSurfaceTransformFlagBitsKHR GetCachedSurfaceTransform() const { return PreTransform; }
diff --git a/Engine/Source/Runtime/VulkanRHI/Private/VulkanViewport.cpp b/Engine/Source/Runtime/VulkanRHI/Private/VulkanViewport.cpp
--- a/Engine/Source/Runtime/VulkanRHI/Private/VulkanViewport.cpp
+++ b/Engine/Source/Runtime/VulkanRHI/Private/VulkanViewport.cpp
@@ -932,11 +932,6 @@ FTextureRHIRef FVulkanDynamicRHI::RHIGetViewportBackBuffer(FRHIViewport* Viewpor
 	check(ViewportRHI);
 	FVulkanViewport* Viewport = ResourceCast(ViewportRHI);
 
-	if (Viewport->SwapChain)
-	{
-		Viewport->SwapChain->RenderThreadPacing();
-	}
-
 	return Viewport->GetBackBuffer(RHICmdList);
 }
 
