VulkanRHI: Correctly set the mutable descriptor support flag.

#rnx

#Synced-CL 49995137

diff --git a/Engine/Source/Runtime/VulkanRHI/Private/VulkanExtensions.cpp b/Engine/Source/Runtime/VulkanRHI/Private/VulkanExtensions.cpp
--- a/Engine/Source/Runtime/VulkanRHI/Private/VulkanExtensions.cpp
+++ b/Engine/Source/Runtime/VulkanRHI/Private/VulkanExtensions.cpp
@@ -1995,6 +1995,7 @@ public:
 	virtual void PostPhysicalDeviceFeatures(FOptionalVulkanDeviceExtensions& ExtensionFlags) override final
 	{
 		bRequirementsPassed = (MutableDescriptorTypeFeatures.mutableDescriptorType == VK_TRUE);
+		ExtensionFlags.HasEXTMutableDescriptorType = bRequirementsPassed;
 	}
 
 	virtual void PreCreateDevice(VkDeviceCreateInfo& DeviceCreateInfo) override final
