VulkanRHI: Fix incorrect stage on backbuffer semaphore.

#rnx

#Synced-CL 49995137

diff --git a/Engine/Source/Runtime/VulkanRHI/Private/VulkanViewport.cpp b/Engine/Source/Runtime/VulkanRHI/Private/VulkanViewport.cpp
--- a/Engine/Source/Runtime/VulkanRHI/Private/VulkanViewport.cpp
+++ b/Engine/Source/Runtime/VulkanRHI/Private/VulkanViewport.cpp
@@ -694,7 +694,7 @@ bool FVulkanViewport::Present(FVulkanCommandListContext& Context, bool bLockToVs
 			if (LIKELY(TryAcquireImageIndex()))
 			{
 				// Wait for semaphore signal before writing to backbuffer image
-				Context.AddWaitSemaphore(VK_PIPELINE_STAGE_BOTTOM_OF_PIPE_BIT, AcquiredSemaphore);
+				Context.AddWaitSemaphore(VK_PIPELINE_STAGE_COLOR_ATTACHMENT_OUTPUT_BIT|VK_PIPELINE_STAGE_TRANSFER_BIT, AcquiredSemaphore);
 				
 				uint32 WindowSizeX = FMath::Min(SizeX, SwapChain->InternalWidth);
 				uint32 WindowSizeY = FMath::Min(SizeY, SwapChain->InternalHeight);
