VulkanRHI: Initialize VK_EXT_device_address_binding_report extension (WIP)

#rnx

#Synced-CL 50114381

diff --git a/Engine/Source/Runtime/VulkanRHI/Private/VulkanDebug.cpp b/Engine/Source/Runtime/VulkanRHI/Private/VulkanDebug.cpp
--- a/Engine/Source/Runtime/VulkanRHI/Private/VulkanDebug.cpp
+++ b/Engine/Source/Runtime/VulkanRHI/Private/VulkanDebug.cpp
@@ -14,9 +14,6 @@ static FAutoConsoleVariable GCVarUniqueValidationMessages(
 
 #define VULKAN_ENABLE_TRACKING_CALLSTACK			1
 
-#define CREATE_MSG_CALLBACK							"vkCreateDebugReportCallbackEXT"
-#define DESTROY_MSG_CALLBACK						"vkDestroyDebugReportCallbackEXT"
-
 DEFINE_LOG_CATEGORY(LogVulkanRHI);
 
 #if VULKAN_HAS_DEBUGGING_ENABLED
@@ -184,7 +181,7 @@ static VKAPI_ATTR VkBool32 VKAPI_CALL DebugUtilsCallback(VkDebugUtilsMessageSeve
 	return VK_FALSE;
 }
 
-void FVulkanDynamicRHI::SetupDebugLayerCallback()
+void FVulkanDynamicRHI::SetupDebugCallbacks()
 {
 	if (ActiveDebugLayerExtension == EActiveDebugLayerExtension::DebugUtilsExtension)
 	{
@@ -206,7 +203,7 @@ void FVulkanDynamicRHI::SetupDebugLayerCallback()
 	}
 }
 
-void FVulkanDynamicRHI::RemoveDebugLayerCallback()
+void FVulkanDynamicRHI::RemoveDebugCallbacks()
 {
 	if (Messenger != VK_NULL_HANDLE)
 	{
diff --git a/Engine/Source/Runtime/VulkanRHI/Private/VulkanDevice.h b/Engine/Source/Runtime/VulkanRHI/Private/VulkanDevice.h
--- a/Engine/Source/Runtime/VulkanRHI/Private/VulkanDevice.h
+++ b/Engine/Source/Runtime/VulkanRHI/Private/VulkanDevice.h
@@ -69,6 +69,7 @@ struct FOptionalVulkanDeviceExtensions
 			uint64 HasEXTShaderObject : 1;
 			uint64 HasEXTGraphicsPipelineLibrary : 1;
 			uint64 HasEXTDeviceGeneratedCommands : 1;
+			uint64 HasEXTDeviceAddressBindingReport : 1;
 
 			// Vendor specific
 			uint64 HasAMDBufferMarker : 1;
@@ -117,13 +118,13 @@ struct FOptionalVulkanDeviceExtensions
 			uint64 HasEXTLoadStoreOpNone : 1;
 			uint64 HasEXTHostImageCopy : 1;
 		};
-		uint64 Packed;
+		uint64 Packed[2];
 	};
 
 	FOptionalVulkanDeviceExtensions()
 	{
 		static_assert(sizeof(Packed) == sizeof(FOptionalVulkanDeviceExtensions), "More bits needed for Packed!");
-		Packed = 0;
+		FMemory::Memzero(Packed);
 	}
 
 	bool HasGPUCrashDumpExtensions() const
diff --git a/Engine/Source/Runtime/VulkanRHI/Private/VulkanDynamicRHI.h b/Engine/Source/Runtime/VulkanRHI/Private/VulkanDynamicRHI.h
--- a/Engine/Source/Runtime/VulkanRHI/Private/VulkanDynamicRHI.h
+++ b/Engine/Source/Runtime/VulkanRHI/Private/VulkanDynamicRHI.h
@@ -442,9 +442,10 @@ protected:
 
 #if VULKAN_HAS_DEBUGGING_ENABLED
 	VkDebugUtilsMessengerEXT Messenger = VK_NULL_HANDLE;
+	VkDebugUtilsMessengerEXT AddressBindingMessenger = VK_NULL_HANDLE;
 
-	void SetupDebugLayerCallback();
-	void RemoveDebugLayerCallback();
+	void SetupDebugCallbacks();
+	void RemoveDebugCallbacks();
 #endif
 
 	FCriticalSection LockBufferCS;
diff --git a/Engine/Source/Runtime/VulkanRHI/Private/VulkanExtensions.cpp b/Engine/Source/Runtime/VulkanRHI/Private/VulkanExtensions.cpp
--- a/Engine/Source/Runtime/VulkanRHI/Private/VulkanExtensions.cpp
+++ b/Engine/Source/Runtime/VulkanRHI/Private/VulkanExtensions.cpp
@@ -155,6 +155,15 @@ TAutoConsoleVariable<int32> GVulkanAllowImageRobustness(
 	ECVF_ReadOnly
 );
 
+TAutoConsoleVariable<int32> GVulkanAllowAddressBindingReport(
+	TEXT("r.Vulkan.AllowAddressBindingReport"),
+	0,
+	TEXT("Whether to allow VK_EXT_device_address_binding_report\n")
+	TEXT("0 to disable (default)\n")
+	TEXT("1 to enable"),
+	ECVF_ReadOnly
+);
+
 #if VULKAN_HAS_DEBUGGING_ENABLED
 extern TAutoConsoleVariable<int32> GGPUValidationCvar;
 #endif
@@ -1474,6 +1483,41 @@ private:
 	VkPhysicalDeviceFaultFeaturesEXT DeviceFaultFeatures;
 };
 
+// ***** VK_EXT_device_address_binding_report
+class FVulkanEXTDeviceAddressBindingReportExtension : public FVulkanDeviceExtension
+{
+public:
+
+	FVulkanEXTDeviceAddressBindingReportExtension(FVulkanDevice* InDevice)
+		: FVulkanDeviceExtension(InDevice, VK_EXT_DEVICE_ADDRESS_BINDING_REPORT_EXTENSION_NAME, VULKAN_EXTENSION_ENABLED)
+	{
+		bEnabledInCode = bEnabledInCode && UE::RHI::UseGPUCrashDebugging() && GVulkanAllowAddressBindingReport.GetValueOnAnyThread();
+	}
+
+	virtual void PrePhysicalDeviceFeatures(VkPhysicalDeviceFeatures2KHR& PhysicalDeviceFeatures2) override final
+	{
+		ZeroVulkanStruct(DeviceAddressBindingReportFeatures, VK_STRUCTURE_TYPE_PHYSICAL_DEVICE_ADDRESS_BINDING_REPORT_FEATURES_EXT);
+		AddToPNext(PhysicalDeviceFeatures2, DeviceAddressBindingReportFeatures);
+	}
+
+	virtual void PostPhysicalDeviceFeatures(FOptionalVulkanDeviceExtensions& ExtensionFlags) override final
+	{
+		bRequirementsPassed = (DeviceAddressBindingReportFeatures.reportAddressBinding == VK_TRUE);
+		ExtensionFlags.HasEXTDeviceAddressBindingReport = bRequirementsPassed;
+	}
+
+	virtual void PreCreateDevice(VkDeviceCreateInfo& DeviceCreateInfo) override final
+	{
+		if (bRequirementsPassed)
+		{
+			AddToPNext(DeviceCreateInfo, DeviceAddressBindingReportFeatures);
+		}
+	}
+
+private:
+	VkPhysicalDeviceAddressBindingReportFeaturesEXT DeviceAddressBindingReportFeatures;
+};
+
 // ***** VK_EXT_host_query_reset
 class FVulkanEXTHostQueryResetExtension : public FVulkanDeviceExtension
 {
@@ -2561,6 +2605,7 @@ FVulkanDeviceExtensionArray FVulkanDeviceExtension::GetUESupportedDeviceExtensio
 	ADD_CUSTOM_EXTENSION(FVulkanEXTCalibratedTimestampsExtension);
 	ADD_CUSTOM_EXTENSION(FVulkanEXTDescriptorBuffer);
 	ADD_CUSTOM_EXTENSION(FVulkanEXTDeviceFaultExtension);
+	ADD_CUSTOM_EXTENSION(FVulkanEXTDeviceAddressBindingReportExtension);
 	ADD_CUSTOM_EXTENSION(FVulkanEXTShaderDemoteToHelperInvocationExtension);
 	ADD_CUSTOM_EXTENSION(FVulkanKHR16BitStorageExtension);
 	ADD_CUSTOM_EXTENSION(FVulkanKHRShaderFloat16Int8Extension);
diff --git a/Engine/Source/Runtime/VulkanRHI/Private/VulkanRHI.cpp b/Engine/Source/Runtime/VulkanRHI/Private/VulkanRHI.cpp
--- a/Engine/Source/Runtime/VulkanRHI/Private/VulkanRHI.cpp
+++ b/Engine/Source/Runtime/VulkanRHI/Private/VulkanRHI.cpp
@@ -559,7 +559,7 @@ void FVulkanDynamicRHI::Shutdown()
 	HMDVulkanExtensions = nullptr;
 
 #if VULKAN_HAS_DEBUGGING_ENABLED
-	RemoveDebugLayerCallback();
+	RemoveDebugCallbacks();
 #endif
 
 	VulkanRHI::vkDestroyInstance(Instance, VULKAN_CPU_ALLOCATOR);
@@ -755,7 +755,7 @@ void FVulkanDynamicRHI::CreateInstance()
 	}
 
 #if VULKAN_HAS_DEBUGGING_ENABLED
-	SetupDebugLayerCallback();
+	SetupDebugCallbacks();
 #endif
 }
 
