VulkanRHI: simplify uniform buffer allocations

#rnx
#Synced-CL 50204094

diff --git a/Engine/Source/Runtime/VulkanRHI/Private/VulkanMemory.cpp b/Engine/Source/Runtime/VulkanRHI/Private/VulkanMemory.cpp
--- a/Engine/Source/Runtime/VulkanRHI/Private/VulkanMemory.cpp
+++ b/Engine/Source/Runtime/VulkanRHI/Private/VulkanMemory.cpp
@@ -3653,7 +3653,8 @@ namespace VulkanRHI
 
 	void FMemoryManager::AllocUniformBuffer(FVulkanAllocation& OutAllocation, uint32 Size)
 	{
-		if (!AllocateBufferPooled(OutAllocation, nullptr, Size, 0, VK_BUFFER_USAGE_UNIFORM_BUFFER_BIT | VK_BUFFER_USAGE_TRANSFER_DST_BIT, VK_MEMORY_PROPERTY_HOST_VISIBLE_BIT | VK_MEMORY_PROPERTY_HOST_COHERENT_BIT, EVulkanAllocationMetaUniformBuffer, __FILE__, __LINE__))
+		const int32 UBAlignment = FMath::Max<uint32>(Device.GetLimits().minUniformBufferOffsetAlignment, 16u);
+		if (!AllocateBufferPooled(OutAllocation, nullptr, Size, UBAlignment, VK_BUFFER_USAGE_UNIFORM_BUFFER_BIT | VK_BUFFER_USAGE_TRANSFER_DST_BIT, VK_MEMORY_PROPERTY_HOST_VISIBLE_BIT | VK_MEMORY_PROPERTY_HOST_COHERENT_BIT, EVulkanAllocationMetaUniformBuffer, __FILE__, __LINE__))
 		{
 			HandleOOM(false);
 			checkNoEntry();
@@ -3663,7 +3664,14 @@ namespace VulkanRHI
 	}
 	void FMemoryManager::FreeUniformBuffer(FVulkanAllocation& InAllocation)
 	{
-		if (InAllocation.HasAllocation())
+		if (InAllocation.bTempBlockAllocator)
+		{
+			Device.GetDeferredDeletionQueue().EnqueueFunc([InAllocation](FVulkanDevice& Device)
+			{
+				Device.GetTempBlockAllocator().FreeReference(InAllocation);
+			});
+		}
+		else if (InAllocation.HasAllocation())
 		{
 			DEC_MEMORY_STAT_BY(STAT_UniformBufferMemory, InAllocation.Size);
 
@@ -3765,7 +3773,7 @@ namespace VulkanRHI
 	}
 
 	FVulkanAllocation::FVulkanAllocation()
-		: bHasOwnership(0), bTransient(false)
+		: bHasOwnership(0), bTransient(false), bTempBlockAllocator(false)
 	{
 		SetType(EVulkanAllocationEmpty);
 	}
@@ -4784,7 +4792,7 @@ namespace VulkanRHI
 		return NewBlock;
 	}
 
-	FTempBlockAllocator::FInternalAlloc FTempBlockAllocator::InternalAlloc(uint32 InSize, FVulkanContextCommon& Context)
+	FTempBlockAllocator::FInternalAlloc FTempBlockAllocator::InternalAlloc(uint32 InSize, FVulkanContextCommon* Context)
 	{
 		const uint32 AlignedSize = Align(InSize, BlockAlignment);
 		checkfSlow(BlockAlignment < BlockSize, TEXT("Requested size of %d (%d aligned) is too large for block size of %d"), InSize, AlignedSize, BlockSize);
@@ -4793,11 +4801,16 @@ namespace VulkanRHI
 		{
 			FReadScopeLock ScopeLock(RWLock);
 
-			if (CurrentBlock->SyncPoints.Contains(Context.GetContextSyncPoint()))
+			if (!Context || CurrentBlock->SyncPoints.Contains(Context->GetContextSyncPoint()))
 			{
 				const uint32 AllocOffset = CurrentBlock->CurrentOffset.fetch_add(AlignedSize);
 				if (AllocOffset + InSize < BlockSize)
 				{
+					if (!Context)
+					{
+						CurrentBlock->References++;
+					}
+
 					return FInternalAlloc{ CurrentBlock, AllocOffset };
 				}
 			}
@@ -4811,7 +4824,14 @@ namespace VulkanRHI
 			uint32 AllocOffset = CurrentBlock->CurrentOffset.fetch_add(AlignedSize);
 			if (AllocOffset + InSize < BlockSize)
 			{
-				CurrentBlock->SyncPoints.AddUnique(Context.GetContextSyncPoint());
+				if (Context)
+				{
+					CurrentBlock->SyncPoints.AddUnique(Context->GetContextSyncPoint());
+				}
+				else
+				{
+					CurrentBlock->References++;
+				}
 				return FInternalAlloc{ CurrentBlock, AllocOffset };
 			}
 
@@ -4826,7 +4846,14 @@ namespace VulkanRHI
 				CurrentBlock = AllocBlock();
 			}
 
-			CurrentBlock->SyncPoints.AddUnique(Context.GetContextSyncPoint());
+			if (Context)
+			{
+				CurrentBlock->SyncPoints.AddUnique(Context->GetContextSyncPoint());
+			}
+			else
+			{
+				CurrentBlock->References++;
+			}
 
 			AllocOffset = CurrentBlock->CurrentOffset.fetch_add(AlignedSize);
 			checkSlow(AllocOffset == 0);
@@ -4837,7 +4864,7 @@ namespace VulkanRHI
 
 	uint8* FTempBlockAllocator::Alloc(uint32 InSize, FVulkanContextCommon& Context, VkDescriptorBufferBindingInfoEXT& OutBindingInfo, VkDeviceSize& OutOffset)
 	{
-		FInternalAlloc Alloc = InternalAlloc(InSize, Context);
+		FInternalAlloc Alloc = InternalAlloc(InSize, &Context);
 
 		ZeroVulkanStruct(OutBindingInfo, VK_STRUCTURE_TYPE_DESCRIPTOR_BUFFER_BINDING_INFO_EXT);
 		OutBindingInfo.address = Alloc.Block->BufferAddress;
@@ -4846,7 +4873,7 @@ namespace VulkanRHI
 		return Alloc.Block->MappedPointer + Alloc.Offset;
 	}
 
-	uint8* FTempBlockAllocator::Alloc(uint32 InSize, uint32 InAlignment, FVulkanContextCommon& Context, FVulkanAllocation& OutAllocation, VkDescriptorAddressInfoEXT* OutDescriptorAddressInfo)
+	uint8* FTempBlockAllocator::Alloc(uint32 InSize, uint32 InAlignment, FVulkanContextCommon* Context, FVulkanAllocation& OutAllocation, VkDescriptorAddressInfoEXT* OutDescriptorAddressInfo)
 	{
 		const uint32 AlignedSize = Align(InSize, BlockAlignment);
 		checkfSlow(AlignedSize < BlockSize, TEXT("Requested size of %d (%d aligned) is too large for block size of %d"), InSize, AlignedSize, BlockSize);
@@ -4857,6 +4884,7 @@ namespace VulkanRHI
 		OutAllocation.VulkanHandle = (uint64)Alloc.Block->Buffer;
 		OutAllocation.Size = InSize;
 		OutAllocation.Offset += Alloc.Offset;
+		OutAllocation.bTempBlockAllocator = (Context == nullptr);
 
 		if (OutDescriptorAddressInfo)
 		{
@@ -4874,13 +4902,42 @@ namespace VulkanRHI
 		checkfSlow(AlignedSize < BlockSize, TEXT("Requested size of %d (%d aligned) is too large for block size of %d"), InSize, AlignedSize, BlockSize);
 		checkSlow(OutStridedDeviceAddressRegion);
 
-		FInternalAlloc Alloc = InternalAlloc(InSize, Context);
+		FInternalAlloc Alloc = InternalAlloc(InSize, &Context);
 		OutStridedDeviceAddressRegion->deviceAddress = Alloc.Block->BufferAddress + Alloc.Offset;
 		OutStridedDeviceAddressRegion->size = InSize;
 		OutStridedDeviceAddressRegion->stride = OutStridedDeviceAddressRegion->size;
 		return Alloc.Block->MappedPointer + Alloc.Offset;
 	}
 
+	void FTempBlockAllocator::FreeReference(const FVulkanAllocation& InAllocation)
+	{
+		FReadScopeLock ScopeLock(RWLock);
+
+		if ((InAllocation.VulkanHandle == (uint64)CurrentBlock->Buffer) &&
+			(InAllocation.AllocatorIndex == CurrentBlock->Allocation.AllocatorIndex) &&
+			(InAllocation.AllocationIndex == CurrentBlock->Allocation.AllocationIndex))
+		{
+			CurrentBlock->References--;
+			return;
+		}
+
+		for (int32 Index = 0; Index < BusyBlocks.Num(); ++Index)
+		{
+			FTempMemoryBlock* Block = BusyBlocks[Index];
+			FVulkanAllocation& BlockAllocation = Block->Allocation;
+			if ((InAllocation.VulkanHandle == (uint64)Block->Buffer) &&
+				(InAllocation.AllocatorIndex == BlockAllocation.AllocatorIndex) &&
+				(InAllocation.AllocationIndex == BlockAllocation.AllocationIndex))
+			{
+				check(Block->References > 0);
+				Block->References--;
+				return;
+			}
+		}
+
+		checkf(false, TEXT("Freed allocation's block was not found in the temp allocator!"));
+	}
+
 	void FTempBlockAllocator::UpdateBlocks()
 	{
 		FWriteScopeLock ScopeLock(RWLock);
@@ -4889,13 +4946,17 @@ namespace VulkanRHI
 		{
 			FTempMemoryBlock* Block = BusyBlocks[Index];
 
-			bool BlockReady = true;
-			for (FVulkanSyncPointRef& SyncPoint : Block->SyncPoints)
+			bool BlockReady = (Block->References == 0);
+
+			if (BlockReady)
 			{
-				if (!SyncPoint->IsComplete())
+				for (FVulkanSyncPointRef& SyncPoint : Block->SyncPoints)
 				{
-					BlockReady = false;
-					break;
+					if (!SyncPoint->IsComplete())
+					{
+						BlockReady = false;
+						break;
+					}
 				}
 			}
 
diff --git a/Engine/Source/Runtime/VulkanRHI/Private/VulkanMemory.h b/Engine/Source/Runtime/VulkanRHI/Private/VulkanMemory.h
--- a/Engine/Source/Runtime/VulkanRHI/Private/VulkanMemory.h
+++ b/Engine/Source/Runtime/VulkanRHI/Private/VulkanMemory.h
@@ -431,11 +431,12 @@ namespace VulkanRHI
 		uint32 AllocationIndex = 0;
 		uint16 AllocatorIndex = 0;
 		EVulkanAllocationMetaType MetaType = EVulkanAllocationMetaUnknown;
-		uint8 Type : 6;
+		uint8 Type : 5;
 		uint8 bHasOwnership : 1;
 		uint8 bTransient : 1;
+		uint8 bTempBlockAllocator : 1;
 
-		static_assert(EVulkanAllocationSize < 128, "Not enough bits to hold EVulkanAllocationType");
+		static_assert(EVulkanAllocationSize < 32, "Not enough bits to hold EVulkanAllocationType");
 		inline EVulkanAllocationType GetType() const
 		{
 			return (EVulkanAllocationType)Type;
@@ -958,6 +959,7 @@ namespace VulkanRHI
 
 			// Used by owners to suballocate
 			TArray<FVulkanSyncPointRef> SyncPoints;
+			std::atomic<uint32> References = 0;
 			std::atomic<uint32> CurrentOffset = 0;
 		};
 
@@ -971,14 +973,21 @@ namespace VulkanRHI
 		virtual ~FTempBlockAllocator();
 
 		uint8* Alloc(uint32 InSize, FVulkanContextCommon& Context, VkDescriptorBufferBindingInfoEXT& OutBindingInfo, VkDeviceSize& OutOffset);
-		uint8* Alloc(uint32 InSize, uint32 InAlignment, FVulkanContextCommon& Context, FVulkanAllocation& OutAllocation, VkDescriptorAddressInfoEXT* OutDescriptorAddressInfo = nullptr);
+		uint8* Alloc(uint32 InSize, uint32 InAlignment, FVulkanContextCommon* Context, FVulkanAllocation& OutAllocation, VkDescriptorAddressInfoEXT* OutDescriptorAddressInfo = nullptr);
 		uint8* Alloc(uint32 InSize, FVulkanContextCommon& Context, VkStridedDeviceAddressRegionKHR* OutStridedDeviceAddressRegion);
 
+		uint8* Alloc(uint32 InSize, uint32 InAlignment, FVulkanContextCommon& Context, FVulkanAllocation& OutAllocation, VkDescriptorAddressInfoEXT* OutDescriptorAddressInfo = nullptr)
+		{
+			return Alloc(InSize, InAlignment, &Context, OutAllocation, OutDescriptorAddressInfo);
+		}
+
+		void FreeReference(const FVulkanAllocation& InAllocation);
+
 		void UpdateBlocks();
 
 	protected:
 		FTempMemoryBlock* AllocBlock();
-		FInternalAlloc InternalAlloc(uint32 InSize, FVulkanContextCommon& Context);
+		FInternalAlloc InternalAlloc(uint32 InSize, FVulkanContextCommon* Context);
 
 		FVulkanDevice& Device;
 
diff --git a/Engine/Source/Runtime/VulkanRHI/Private/VulkanResources.h b/Engine/Source/Runtime/VulkanRHI/Private/VulkanResources.h
--- a/Engine/Source/Runtime/VulkanRHI/Private/VulkanResources.h
+++ b/Engine/Source/Runtime/VulkanRHI/Private/VulkanResources.h
@@ -972,7 +972,7 @@ public:
 	const TArray<TRefCountPtr<FRHIResource>>& GetResourceTable() const { return ResourceTable; }
 
 	void UpdateResourceTable(const FRHIUniformBufferLayout& InLayout, const void* Contents, int32 ResourceNum);
-	void UpdateResourceTable(FRHIResource** Resources, int32 ResourceNum);
+	void UpdateResourceTable(TArrayView<FRHIResource*> Resources);
 
 	VkBuffer GetBufferHandle() const
 	{
@@ -989,12 +989,6 @@ public:
 		return Allocation;
 	}
 
-	void UpdateAllocation(VulkanRHI::FVulkanAllocation& NewAlloc)
-	{
-		NewAlloc.Swap(Allocation);
-		UpdateBindlessHandle();
-	}
-
 	bool IsUniformView() const
 	{
 		return UniformViewSRV != nullptr;
@@ -1002,23 +996,24 @@ public:
 
 	FRHIDescriptorHandle GetBindlessHandle() const override
 	{
-		checkfSlow(!BindlessHandle.IsValid() || (CachedDeviceAddress == GetDeviceAddress()), TEXT("Uniform Buffer bindless handle is out of date!"));
 		return BindlessHandle;
 	}
 	
-	void UpdateBindlessHandle();
+	void UpdateBindlessHandle(const FVulkanContextArray& Contexts);
 	VkDeviceAddress GetDeviceAddress() const;
 
 	void SetupUniformBufferView();
 
+	void AllocateMemory(const void* InitialData);
+
 public:
-	FVulkanDevice* Device;
 	const EUniformBufferUsage Usage;
 
 protected:
+
+	FVulkanDevice& Device;
 	VulkanRHI::FVulkanAllocation Allocation;
 	FRHIDescriptorHandle BindlessHandle;
-	VkDeviceAddress CachedDeviceAddress = 0;
 	FRHIShaderResourceView* UniformViewSRV = nullptr;
 };
 
diff --git a/Engine/Source/Runtime/VulkanRHI/Private/VulkanUniformBuffer.cpp b/Engine/Source/Runtime/VulkanRHI/Private/VulkanUniformBuffer.cpp
--- a/Engine/Source/Runtime/VulkanRHI/Private/VulkanUniformBuffer.cpp
+++ b/Engine/Source/Runtime/VulkanRHI/Private/VulkanUniformBuffer.cpp
@@ -27,56 +27,32 @@ static FAutoConsoleVariableRef CVarVulkanAllowUniformUpload(
 	Uniform buffer RHI object
 -----------------------------------------------------------------------------*/
 
-static void UpdateUniformBufferConstants(FVulkanDevice& Device, void* DestinationData, const void* SourceData, const FRHIUniformBufferLayout* Layout)
-{
-	UE::RHICore::UpdateUniformBufferConstants(DestinationData, SourceData, *Layout, Device.SupportsBindless());
-}
 
-static bool UseTemporaryBuffer(EUniformBufferUsage Usage)
+void FVulkanUniformBuffer::AllocateMemory(const void* InitialData)
 {
-	// Add a cvar to control this behavior?
-	return (Usage == UniformBuffer_SingleDraw || Usage == UniformBuffer_SingleFrame);
-}
+	const int32 DataSize = GetLayout().ConstantBufferSize;
 
-static void UpdateUniformBufferHelper(FVulkanCommandListContext& Context, FVulkanUniformBuffer& VulkanUniformBuffer, const void* Data, bool bUpdateConstants = true)
-{
-	FVulkanCommandBuffer* CmdBuffer = Context.GetActiveCmdBuffer();
+	// Free any previous allocations
+	Device.GetMemoryManager().FreeUniformBuffer(Allocation);
 
-	FVulkanDevice& Device = Context.Device;
-	const int32 DataSize = VulkanUniformBuffer.GetLayout().ConstantBufferSize;
-	const int32 DataAlignment = FMath::Max<uint32>(Device.GetLimits().minUniformBufferOffsetAlignment, 16u);
-
-	VulkanRHI::FVulkanAllocation TempAllocation;
-	void* DestinationData = Device.GetTempBlockAllocator().Alloc(DataSize, DataAlignment, Context, TempAllocation);
-
-	if (bUpdateConstants)
+	void* DestinationData = nullptr;
+	if (Usage == UniformBuffer_SingleDraw || Usage == UniformBuffer_SingleFrame)
 	{
-		// Update constants as the data is copied
-		UpdateUniformBufferConstants(Device, DestinationData, Data, VulkanUniformBuffer.GetLayoutPtr());
+		const int32 UBAlignment = FMath::Max<uint32>(Device.GetLimits().minUniformBufferOffsetAlignment, 16u);
+		DestinationData = Device.GetTempBlockAllocator().Alloc(DataSize, UBAlignment, nullptr, Allocation);
 	}
 	else
 	{
-		// Don't touch constant, copy the data as-is
-		FMemory::Memcpy(DestinationData, Data, DataSize);
+		Device.GetMemoryManager().AllocUniformBuffer(Allocation, DataSize);
+		DestinationData = Allocation.GetMappedPointer(&Device);
 	}
 
-	if (UseTemporaryBuffer(VulkanUniformBuffer.Usage))
+	if (DestinationData && InitialData)
 	{
-		VulkanUniformBuffer.UpdateAllocation(TempAllocation);
+		UE::RHICore::UpdateUniformBufferConstants(DestinationData, InitialData, GetLayout(), Device.SupportsBindless());
+		Allocation.FlushMappedMemory(&Device);
 	}
-	else
-	{
-		check(CmdBuffer->IsOutsideRenderPass());
-
-		VkBufferCopy Region;
-		Region.size = DataSize;
-		Region.srcOffset = TempAllocation.Offset;
-		Region.dstOffset = VulkanUniformBuffer.GetOffset();
-		VkBuffer UBBuffer = VulkanUniformBuffer.GetBufferHandle();
-
-		VulkanRHI::vkCmdCopyBuffer(CmdBuffer->GetHandle(), TempAllocation.GetBufferHandle(), UBBuffer, 1, &Region);
-	}
-};
+}
 
 void FVulkanUniformBuffer::SetupUniformBufferView()
 {
@@ -93,8 +69,8 @@ void FVulkanUniformBuffer::SetupUniformBufferView()
 
 FVulkanUniformBuffer::FVulkanUniformBuffer(FVulkanDevice& InDevice, const FRHIUniformBufferLayout* InLayout, const void* Contents, EUniformBufferUsage InUsage, EUniformBufferValidation Validation)
 	: FRHIUniformBuffer(InLayout)
-	, Device(&InDevice)
 	, Usage(InUsage)
+	, Device(InDevice)
 {
 #if VULKAN_ENABLE_AGGRESSIVE_STATS
 	SCOPE_CYCLE_COUNTER(STAT_VulkanUniformBufferCreateTime);
@@ -137,53 +113,8 @@ FVulkanUniformBuffer::FVulkanUniformBuffer(FVulkanDevice& InDevice, const FRHIUn
 
 	if (InLayout->ConstantBufferSize > 0)
 	{
-		const bool bInRenderingThread = IsInRenderingThread();
-		const bool bInRHIThread = IsInRHIThread();
-
-PRAGMA_DISABLE_DEPRECATION_WARNINGS
-		if (UseTemporaryBuffer(InUsage) && (bInRenderingThread || bInRHIThread)
-			// :todo-jn:  Temporary check until we have a command list arg passed in to avoid a race where the RenderThread
-			// would pick up other tasks (because of task retraction) and execute them as if on the RenderThread.
-			&& !UE::Tasks::Private::IsThreadRetractingTask())
-PRAGMA_ENABLE_DEPRECATION_WARNINGS
-		{
-			if (Contents)
-			{
-				FRHICommandListImmediate& RHICmdList = FRHICommandListExecutor::GetImmediateCommandList();
-				int32 DataSize = InLayout->ConstantBufferSize;
-
-				// make sure we allocate from RingBuffer on RHIT
-				const bool bCanAllocOnThisThread = RHICmdList.Bypass() || (!IsRunningRHIInSeparateThread() && bInRenderingThread) || bInRHIThread;
-				if (bCanAllocOnThisThread)
-				{
-					UpdateUniformBufferHelper(Device->GetImmediateContext(), *this, Contents);
-				}
-				else
-				{
-					void* CmdListConstantBufferData = RHICmdList.Alloc(DataSize, 16);
-					UpdateUniformBufferConstants(*Device, CmdListConstantBufferData, Contents, InLayout);
-
-					RHICmdList.EnqueueLambda([UniformBuffer=this, DataSize, CmdListConstantBufferData](FRHICommandList& ExecutingCmdList)
-					{
-						UpdateUniformBufferHelper(FVulkanCommandListContext::Get(ExecutingCmdList), *UniformBuffer, CmdListConstantBufferData, false);
-					});
-
-					RHICmdList.RHIThreadFence(true);
-				}
-			}
-		}
-		else
-		{
-			VulkanRHI::FMemoryManager& ResourceMgr = Device->GetMemoryManager();
-			// Set it directly as there is no previous one
-			ResourceMgr.AllocUniformBuffer(Allocation, InLayout->ConstantBufferSize);
-			if (Contents)
-			{
-				UpdateUniformBufferConstants(*Device, Allocation.GetMappedPointer(Device), Contents, InLayout);
-				Allocation.FlushMappedMemory(Device);
-			}
-			UpdateBindlessHandle();
-		}
+		AllocateMemory(Contents);
+		BindlessHandle = Device.GetBindlessDescriptorManager()->AllocateDescriptor(ERHIDescriptorType::BufferSRV, GetDeviceAddress(), GetSize());
 	}
 }
 
@@ -191,10 +122,10 @@ FVulkanUniformBuffer::~FVulkanUniformBuffer()
 {
 	if (BindlessHandle.IsValid())
 	{
-		Device->GetDeferredDeletionQueue().EnqueueBindlessHandle(BindlessHandle);
+		Device.GetDeferredDeletionQueue().EnqueueBindlessHandle(BindlessHandle);
 	}
 
-	Device->GetMemoryManager().FreeUniformBuffer(Allocation);
+	Device.GetMemoryManager().FreeUniformBuffer(Allocation);
 }
 
 void FVulkanUniformBuffer::UpdateResourceTable(const FRHIUniformBufferLayout& InLayout, const void* Contents, int32 NumResources)
@@ -208,32 +139,21 @@ void FVulkanUniformBuffer::UpdateResourceTable(const FRHIUniformBufferLayout& In
 	}
 }
 
-void FVulkanUniformBuffer::UpdateResourceTable(FRHIResource** Resources, int32 ResourceNum)
+void FVulkanUniformBuffer::UpdateResourceTable(TArrayView<FRHIResource*> Resources)
 {
-	check(ResourceTable.Num() == ResourceNum);
-
-	for (int32 ResourceIndex = 0; ResourceIndex < ResourceNum; ++ResourceIndex)
+	check(Resources.Num() <= ResourceTable.Num());
+	for (int32 ResourceIndex = 0; ResourceIndex < Resources.Num(); ++ResourceIndex)
 	{
 		ResourceTable[ResourceIndex] = Resources[ResourceIndex];
 	}
 }
 
-void FVulkanUniformBuffer::UpdateBindlessHandle()
+void FVulkanUniformBuffer::UpdateBindlessHandle(const FVulkanContextArray& Contexts)
 {
-	if (Device->GetBindlessDescriptorManager()->IsSupported())
+	if (BindlessHandle.IsValid() && Device.GetBindlessDescriptorManager()->IsSupported())
 	{
-		const VkDeviceAddress CurrentAddress = GetDeviceAddress();
-		if (!BindlessHandle.IsValid() || (CachedDeviceAddress == 0) || (CurrentAddress != CachedDeviceAddress))
-		{
-			if (BindlessHandle.IsValid())
-			{
-				Device->GetDeferredDeletionQueue().EnqueueBindlessHandle(BindlessHandle);
-			}
-
-			BindlessHandle = Device->GetBindlessDescriptorManager()->AllocateDescriptor(ERHIDescriptorType::BufferSRV, CurrentAddress, GetSize());
-
-			CachedDeviceAddress = CurrentAddress;
-		}
+		const VkDeviceAddress BufferAddress = GetDeviceAddress();
+		Device.GetBindlessDescriptorManager()->UpdateDescriptor(Contexts, BindlessHandle, BufferAddress, GetSize(), false);
 	}
 }
 
@@ -244,7 +164,7 @@ VkDeviceAddress FVulkanUniformBuffer::GetDeviceAddress() const
 		.pNext = nullptr,
 		.buffer = GetBufferHandle()
 	};
-	VkDeviceAddress BufferAddress = VulkanRHI::vkGetBufferDeviceAddressKHR(Device->GetHandle(), &BufferInfo);
+	VkDeviceAddress BufferAddress = VulkanRHI::vkGetBufferDeviceAddressKHR(Device.GetHandle(), &BufferInfo);
 	BufferAddress += GetOffset();
 	return BufferAddress;
 }
@@ -263,86 +183,52 @@ void FVulkanDynamicRHI::RHIUpdateUniformBuffer(FRHICommandListBase& RHICmdList,
 	SCOPE_CYCLE_COUNTER(STAT_VulkanUpdateUniformBuffers);
 
 	const FRHIUniformBufferLayout& Layout = UniformBuffer->GetLayout();
-
 	const int32 ConstantBufferSize = Layout.ConstantBufferSize;
-	const int32 NumResources = Layout.Resources.Num();
+	const bool bBypass = RHICmdList.Bypass();
 
-	VulkanRHI::FVulkanAllocation NewUBAlloc;
-	bool bUseUpload = GVulkanAllowUniformUpload && !RHICmdList.IsInsideRenderPass(); //inside renderpasses, a rename is enforced.
-	const bool bUseTempBuffer = UseTemporaryBuffer(UniformBuffer->Usage);
-
-	if (!bUseUpload && !bUseTempBuffer)
+	TArrayView<FRHIResource*> CmdListResources;
+	if (const int32 NumResources = Layout.Resources.Num())
 	{
-		if (ConstantBufferSize > 0)
+		CmdListResources = bBypass
+			? TArrayView<FRHIResource*>((FRHIResource**)FMemory_Alloca(sizeof(FRHIResource*) * NumResources), NumResources)
+			: RHICmdList.AllocArrayUninitialized<FRHIResource*>(NumResources);
+
+		for (int32 Index = 0; Index < NumResources; Index++)
 		{
-			SCOPE_CYCLE_COUNTER(STAT_VulkanUpdateUniformBuffersRename);
-			Device->GetMemoryManager().AllocUniformBuffer(NewUBAlloc, ConstantBufferSize);
-			if (Contents)
-			{
-				UpdateUniformBufferConstants(*Device, NewUBAlloc.GetMappedPointer(Device), Contents, &Layout);
-				NewUBAlloc.FlushMappedMemory(Device);
-			}
+			const FRHIUniformBufferResource& Parameter = Layout.Resources[Index];
+			CmdListResources[Index] = GetShaderParameterResourceRHI(Contents, Parameter.MemberOffset, Parameter.MemberType);
 		}
 	}
 
-	const bool bRHIBypass = RHICmdList.Bypass();
-	if (bRHIBypass)
+	if (!bBypass)
 	{
-		if (ConstantBufferSize > 0)
+		for (FRHIResource* Resource : CmdListResources)
 		{
-			if (bUseUpload || bUseTempBuffer)
-			{			
-				FVulkanCommandListContext& Context = Device->GetImmediateContext();
-				UpdateUniformBufferHelper(Context, *UniformBuffer, Contents);
-			}
-			else
-			{
-				UniformBuffer->UpdateAllocation(NewUBAlloc);
-				Device->GetMemoryManager().FreeUniformBuffer(NewUBAlloc);
-			}
+			Resource->AddRef();
 		}
-
-		UniformBuffer->UpdateResourceTable(Layout, Contents, NumResources);
 	}
-	else
-	{
-		FRHIResource** CmdListResources = nullptr;
-		if (NumResources > 0)
-		{
-			CmdListResources = (FRHIResource**)RHICmdList.Alloc(sizeof(FRHIResource*) * NumResources, alignof(FRHIResource*));
 
-			for (int32 Index = 0; Index < NumResources; ++Index)
-			{
-				CmdListResources[Index] = GetShaderParameterResourceRHI(Contents, Layout.Resources[Index].MemberOffset, Layout.Resources[Index].MemberType);
-			}
-		}
+	const void* const TempContents = (bBypass || ConstantBufferSize == 0)
+		? Contents
+		: RHICmdList.AllocCopy(Contents, ConstantBufferSize, 16);
 
-		if (bUseUpload || bUseTempBuffer)
+	RHICmdList.EnqueueLambdaMultiPipe(GetEnabledRHIPipelines(), FRHICommandListBase::EThreadFence::Enabled, TEXT("RHIUpdateUniformBuffer"),
+		[UniformBuffer, TempContents, ConstantBufferSize, bBypass, CmdListResources = MoveTemp(CmdListResources)](const FVulkanContextArray& Contexts)
+	{
+		if (ConstantBufferSize > 0)
 		{
-			void* CmdListConstantBufferData = RHICmdList.Alloc(ConstantBufferSize, 16);
-			FMemory::Memcpy(CmdListConstantBufferData, Contents, ConstantBufferSize);
-
-			RHICmdList.EnqueueLambda([UniformBuffer, CmdListResources, NumResources, ConstantBufferSize, CmdListConstantBufferData](FRHICommandListBase& ExecutingCmdList)
-			{
-				UpdateUniformBufferHelper(FVulkanCommandListContext::Get(ExecutingCmdList), *UniformBuffer, CmdListConstantBufferData);
-				UniformBuffer->UpdateResourceTable(CmdListResources, NumResources);
-			});
+			UniformBuffer->AllocateMemory(TempContents);
+			UniformBuffer->UpdateBindlessHandle(Contexts);
 		}
-		else
+
+		UniformBuffer->UpdateResourceTable(CmdListResources);
+		if (!bBypass)
 		{
-			NewUBAlloc.Disown(); //this releases ownership while its put into the lambda
-			RHICmdList.EnqueueLambda([UniformBuffer, NewUBAlloc, CmdListResources, NumResources](FRHICommandListBase& CmdList)
+			for (FRHIResource* Resource : CmdListResources)
 			{
-				VulkanRHI::FVulkanAllocation Alloc;
-				Alloc.Reference(NewUBAlloc);
-				Alloc.Own(); //this takes ownership of the allocation
-				UniformBuffer->UpdateAllocation(Alloc);
-				UniformBuffer->Device->GetMemoryManager().FreeUniformBuffer(Alloc);
-				UniformBuffer->UpdateResourceTable(CmdListResources, NumResources);
-			});
+				Resource->Release();
+			}
 		}
-		
-		RHICmdList.RHIThreadFence(true);
-	}
+	});
 }
 
