Vulkan Bindless
- VkDescriptorType -> DescriptorType should use masks since buffer srvs and uniform buffers overlap.
- Changing FVulkanUniformBuffer descriptor type to CBV.

#Synced-CL 50213537

#Synced-CL 50213537

diff --git a/Engine/Source/Runtime/VulkanRHI/Private/VulkanBindlessDescriptorManager.cpp b/Engine/Source/Runtime/VulkanRHI/Private/VulkanBindlessDescriptorManager.cpp
--- a/Engine/Source/Runtime/VulkanRHI/Private/VulkanBindlessDescriptorManager.cpp
+++ b/Engine/Source/Runtime/VulkanRHI/Private/VulkanBindlessDescriptorManager.cpp
@@ -131,6 +131,7 @@ static constexpr uint8 GetIndexForDescriptorType(ERHIDescriptorType DescriptorTy
 	case ERHIDescriptorType::TypedBufferUAV:        return VulkanBindless::BindlessStorageTexelBufferSet;
 	case ERHIDescriptorType::BufferUAV:             return VulkanBindless::BindlessStorageBufferSet;
 	case ERHIDescriptorType::BufferSRV:             return VulkanBindless::BindlessUniformBufferSet;
+	case ERHIDescriptorType::CBV:                   return VulkanBindless::BindlessUniformBufferSet;
 	case ERHIDescriptorType::AccelerationStructure: return VulkanBindless::BindlessAccelerationStructureSet;
 	default: checkNoEntry();
 	}
@@ -139,22 +140,22 @@ static constexpr uint8 GetIndexForDescriptorType(ERHIDescriptorType DescriptorTy
 }
 
 
-static constexpr ERHIDescriptorType GetUEDescriptorTypeFromVkDescriptorType(VkDescriptorType DescriptorType)
+static constexpr ERHIDescriptorTypeMask GetUEDescriptorTypeMaskFromVkDescriptorType(VkDescriptorType DescriptorType)
 {
 	switch (DescriptorType)
 	{
-	case VK_DESCRIPTOR_TYPE_SAMPLER:                    return ERHIDescriptorType::Sampler;
-	case VK_DESCRIPTOR_TYPE_SAMPLED_IMAGE:              return ERHIDescriptorType::TextureSRV;
-	case VK_DESCRIPTOR_TYPE_STORAGE_IMAGE:              return ERHIDescriptorType::TextureUAV;
-	case VK_DESCRIPTOR_TYPE_UNIFORM_TEXEL_BUFFER:       return ERHIDescriptorType::TypedBufferSRV;
-	case VK_DESCRIPTOR_TYPE_STORAGE_TEXEL_BUFFER:       return ERHIDescriptorType::TypedBufferUAV;
-	case VK_DESCRIPTOR_TYPE_STORAGE_BUFFER:             return ERHIDescriptorType::BufferUAV;
-	case VK_DESCRIPTOR_TYPE_UNIFORM_BUFFER:             return ERHIDescriptorType::BufferSRV;
-	case VK_DESCRIPTOR_TYPE_ACCELERATION_STRUCTURE_KHR: return ERHIDescriptorType::AccelerationStructure;
+	case VK_DESCRIPTOR_TYPE_SAMPLER:                    return ERHIDescriptorTypeMask::Sampler;
+	case VK_DESCRIPTOR_TYPE_SAMPLED_IMAGE:              return ERHIDescriptorTypeMask::TextureSRV;
+	case VK_DESCRIPTOR_TYPE_STORAGE_IMAGE:              return ERHIDescriptorTypeMask::TextureUAV;
+	case VK_DESCRIPTOR_TYPE_UNIFORM_TEXEL_BUFFER:       return ERHIDescriptorTypeMask::TypedBufferSRV;
+	case VK_DESCRIPTOR_TYPE_STORAGE_TEXEL_BUFFER:       return ERHIDescriptorTypeMask::TypedBufferUAV;
+	case VK_DESCRIPTOR_TYPE_STORAGE_BUFFER:             return ERHIDescriptorTypeMask::BufferUAV;
+	case VK_DESCRIPTOR_TYPE_UNIFORM_BUFFER:             return ERHIDescriptorTypeMask::BufferSRV | ERHIDescriptorTypeMask::CBV;
+	case VK_DESCRIPTOR_TYPE_ACCELERATION_STRUCTURE_KHR: return ERHIDescriptorTypeMask::AccelerationStructure;
 	default: checkNoEntry();
 	}
 
-	return ERHIDescriptorType::Invalid;
+	return ERHIDescriptorTypeMask::None;
 }
 
 static constexpr VkDescriptorType GetVkDescriptorTypeFromUEDescriptorType(ERHIDescriptorType DescriptorType)
@@ -168,6 +169,7 @@ static constexpr VkDescriptorType GetVkDescriptorTypeFromUEDescriptorType(ERHIDe
 	case ERHIDescriptorType::TypedBufferUAV:        return VK_DESCRIPTOR_TYPE_STORAGE_TEXEL_BUFFER;
 	case ERHIDescriptorType::BufferUAV:             return VK_DESCRIPTOR_TYPE_STORAGE_BUFFER;
 	case ERHIDescriptorType::BufferSRV:             return VK_DESCRIPTOR_TYPE_UNIFORM_BUFFER;
+	case ERHIDescriptorType::CBV:                   return VK_DESCRIPTOR_TYPE_UNIFORM_BUFFER;
 	case ERHIDescriptorType::AccelerationStructure: return VK_DESCRIPTOR_TYPE_ACCELERATION_STRUCTURE_KHR;
 	default: checkNoEntry();
 	}
@@ -484,7 +486,7 @@ void FVulkanBindlessDescriptorManager::Init()
 
 			for (VkDescriptorType DescriptorType : InOutState.DescriptorTypes)
 			{
-				EnumAddFlags(TypeMask, RHIDescriptorTypeMaskFromType(GetUEDescriptorTypeFromVkDescriptorType(DescriptorType)));
+				EnumAddFlags(TypeMask, GetUEDescriptorTypeMaskFromVkDescriptorType(DescriptorType));
 				AllocatorStats.Emplace(GetStatForDescriptorType(DescriptorType));
 
 				if (DescriptorType == VK_DESCRIPTOR_TYPE_SAMPLER)
@@ -916,7 +918,7 @@ void FVulkanBindlessDescriptorManager::UpdateDescriptor(const FVulkanContextArra
 {
 	if (IsSupported())
 	{
-		check(DescriptorHandle.GetType() == ERHIDescriptorType::BufferSRV || DescriptorHandle.GetType() == ERHIDescriptorType::BufferUAV);
+		check(DescriptorHandle.GetType() == ERHIDescriptorType::BufferSRV || DescriptorHandle.GetType() == ERHIDescriptorType::CBV || DescriptorHandle.GetType() == ERHIDescriptorType::BufferUAV);
 
 		VkDescriptorAddressInfoEXT AddressInfo;
 		ZeroVulkanStruct(AddressInfo, VK_STRUCTURE_TYPE_DESCRIPTOR_ADDRESS_INFO_EXT);
diff --git a/Engine/Source/Runtime/VulkanRHI/Private/VulkanUniformBuffer.cpp b/Engine/Source/Runtime/VulkanRHI/Private/VulkanUniformBuffer.cpp
--- a/Engine/Source/Runtime/VulkanRHI/Private/VulkanUniformBuffer.cpp
+++ b/Engine/Source/Runtime/VulkanRHI/Private/VulkanUniformBuffer.cpp
@@ -114,7 +114,7 @@ FVulkanUniformBuffer::FVulkanUniformBuffer(FVulkanDevice& InDevice, const FRHIUn
 	if (InLayout->ConstantBufferSize > 0)
 	{
 		AllocateMemory(Contents);
-		BindlessHandle = Device.GetBindlessDescriptorManager()->AllocateDescriptor(ERHIDescriptorType::BufferSRV, GetDeviceAddress(), GetSize());
+		BindlessHandle = Device.GetBindlessDescriptorManager()->AllocateDescriptor(ERHIDescriptorType::CBV, GetDeviceAddress(), GetSize());
 	}
 }
 
