VulkanRHI: Use mutable descriptors to have a single descriptor buffer for all resources

#rnx

#Synced-CL 50218591

diff --git a/Engine/Shaders/Public/Platform/Vulkan/VulkanCommon.ush b/Engine/Shaders/Public/Platform/Vulkan/VulkanCommon.ush
--- a/Engine/Shaders/Public/Platform/Vulkan/VulkanCommon.ush
+++ b/Engine/Shaders/Public/Platform/Vulkan/VulkanCommon.ush
@@ -7,7 +7,7 @@
 #pragma once
 
 // Update this GUID to invalidate shader recompilation for only Vulkan shaders
-#pragma message("UESHADERMETADATA_VERSION 8F516D4E-001A-4398-87CD-BB6D3450CF1D")
+#pragma message("UESHADERMETADATA_VERSION A6FA7DFC-55D8-4F8C-A56C-C3B3AD4645C3")
 
 #define COMPILER_SUPPORTS_ULONG_TYPES 1
 #define COMPILER_SUPPORTS_UINT64_IMAGE_ATOMICS (1 && COMPILER_SUPPORTS_ULONG_TYPES)
diff --git a/Engine/Source/Runtime/VulkanRHI/Private/VulkanBindlessDescriptorManager.cpp b/Engine/Source/Runtime/VulkanRHI/Private/VulkanBindlessDescriptorManager.cpp
--- a/Engine/Source/Runtime/VulkanRHI/Private/VulkanBindlessDescriptorManager.cpp
+++ b/Engine/Source/Runtime/VulkanRHI/Private/VulkanBindlessDescriptorManager.cpp
@@ -13,66 +13,16 @@ static FAutoConsoleVariableRef CVarVulkanBindlessMaxSamplerDescriptorCount(
 	ECVF_ReadOnly
 );
 
-int32 GVulkanBindlessMaxSampledImageDescriptorCount = 256 * 1024;
-static FAutoConsoleVariableRef CVarVulkanBindlessMaxSampledImageCount(
-	TEXT("r.Vulkan.Bindless.MaxResourceSampledImageCount"),
-	GVulkanBindlessMaxSampledImageDescriptorCount,
-	TEXT("Maximum bindless Sampled Image descriptor count"),
-	ECVF_ReadOnly
-);
-
-int32 GVulkanBindlessMaxStorageImageDescriptorCount = 64 * 1024;
-static FAutoConsoleVariableRef CVarVulkanBindlessMaxStorageImageCount(
-	TEXT("r.Vulkan.Bindless.MaxResourceStorageImageCount"),
-	GVulkanBindlessMaxStorageImageDescriptorCount,
-	TEXT("Maximum bindless Storage Image descriptor count"),
-	ECVF_ReadOnly
-);
-
-int32 GVulkanBindlessMaxUniformTexelBufferDescriptorCount = 64 * 1024;
-static FAutoConsoleVariableRef CVarVulkanBindlessMaxUniformTexelBufferCount(
-	TEXT("r.Vulkan.Bindless.MaxResourceUniformTexelBufferCount"),
-	GVulkanBindlessMaxUniformTexelBufferDescriptorCount,
-	TEXT("Maximum bindless Uniform Texel Buffer descriptor count"),
-	ECVF_ReadOnly
-);
-
-int32 GVulkanBindlessMaxStorageTexelBufferDescriptorCount = 64 * 1024;
-static FAutoConsoleVariableRef CVarVulkanBindlessMaxStorageTexelBufferCount(
-	TEXT("r.Vulkan.Bindless.MaxResourceStorageTexelBufferCount"),
-	GVulkanBindlessMaxStorageTexelBufferDescriptorCount,
-	TEXT("Maximum bindless Storage Texel Buffer descriptor count"),
-	ECVF_ReadOnly
-);
-
-int32 GVulkanBindlessMaxUniformBufferDescriptorCount = 64 * 1024;
-static FAutoConsoleVariableRef CVarVulkanBindlessMaxUniformBufferCount(
-	TEXT("r.Vulkan.Bindless.MaxResourceUniformBufferCount"),
-	GVulkanBindlessMaxUniformBufferDescriptorCount,
-	TEXT("Maximum bindless Uniform Buffer descriptor count"),
-	ECVF_ReadOnly
-);
-
-int32 GVulkanBindlessMaxStorageBufferDescriptorCount = 64 * 1024;
-static FAutoConsoleVariableRef CVarVulkanBindlessMaxStorageBufferCount(
-	TEXT("r.Vulkan.Bindless.MaxResourceStorageBufferCount"),
-	GVulkanBindlessMaxStorageBufferDescriptorCount,
-	TEXT("Maximum bindless Storage Buffer descriptor count"),
-	ECVF_ReadOnly
-);
-
-int32 GVulkanBindlessMaxAccelerationStructureDescriptorCount = 64 * 1024;
-static FAutoConsoleVariableRef CVarVulkanBindlessMaxAccelerationStructureCount(
-	TEXT("r.Vulkan.Bindless.MaxResourceAccelerationStructureCount"),
-	GVulkanBindlessMaxAccelerationStructureDescriptorCount,
-	TEXT("Maximum bindless Acceleration Structure descriptor count"),
+int32 GVulkanBindlessMaxResourceDescriptorCount = 1024 * 1024;
+static FAutoConsoleVariableRef CVarVulkanBindlessMaxResourceDescriptorCount(
+	TEXT("r.Vulkan.Bindless.MaxResourceDescriptorCount"),
+	GVulkanBindlessMaxResourceDescriptorCount,
+	TEXT("Maximum bindless resource descriptor count"),
 	ECVF_ReadOnly
 );
 
 DECLARE_STATS_GROUP(TEXT("Vulkan Bindless"), STATGROUP_VulkanBindless, STATCAT_Advanced);
 
-DECLARE_DWORD_ACCUMULATOR_STAT(TEXT("Descriptors Allocated"), STAT_VulkanBindlessDescriptorsAllocated, STATGROUP_VulkanBindless);
-
 DECLARE_DWORD_ACCUMULATOR_STAT(TEXT("Sampler Descriptors Allocated"),               STAT_VulkanBindlessSamplerDescriptorsAllocated,               STATGROUP_VulkanBindless);
 DECLARE_DWORD_ACCUMULATOR_STAT(TEXT("SampledImage Descriptors Allocated"),          STAT_VulkanBindlessSampledImageDescriptorsAllocated,          STATGROUP_VulkanBindless);
 DECLARE_DWORD_ACCUMULATOR_STAT(TEXT("StorageImage Descriptors Allocated"),          STAT_VulkanBindlessStorageImageDescriptorsAllocated,          STATGROUP_VulkanBindless);
@@ -84,62 +34,25 @@ DECLARE_DWORD_ACCUMULATOR_STAT(TEXT("AccelerationStructure Descriptors Allocated
 
 DECLARE_DWORD_COUNTER_STAT(TEXT("Write Per Frame"), STAT_VulkanBindlessWritePerFrame, STATGROUP_VulkanBindless);
 
-static TStatId GetStatForDescriptorType(VkDescriptorType DescriptorType)
+static TStatId GetStatForDescriptorType(ERHIDescriptorType DescriptorType)
 {
 	switch (DescriptorType)
 	{
-	case VK_DESCRIPTOR_TYPE_SAMPLER:                    return GET_STATID(STAT_VulkanBindlessSamplerDescriptorsAllocated);
-	case VK_DESCRIPTOR_TYPE_SAMPLED_IMAGE:              return GET_STATID(STAT_VulkanBindlessSampledImageDescriptorsAllocated);
-	case VK_DESCRIPTOR_TYPE_STORAGE_IMAGE:              return GET_STATID(STAT_VulkanBindlessStorageImageDescriptorsAllocated);
-	case VK_DESCRIPTOR_TYPE_UNIFORM_TEXEL_BUFFER:       return GET_STATID(STAT_VulkanBindlessUniformTexelBufferDescriptorsAllocated);
-	case VK_DESCRIPTOR_TYPE_STORAGE_TEXEL_BUFFER:       return GET_STATID(STAT_VulkanBindlessStorageTexelBufferDescriptorsAllocated);
-	case VK_DESCRIPTOR_TYPE_STORAGE_BUFFER:             return GET_STATID(STAT_VulkanBindlessUniformBufferDescriptorsAllocated);
-	case VK_DESCRIPTOR_TYPE_UNIFORM_BUFFER:             return GET_STATID(STAT_VulkanBindlessStorageBufferDescriptorsAllocated);
-	case VK_DESCRIPTOR_TYPE_ACCELERATION_STRUCTURE_KHR: return GET_STATID(STAT_VulkanBindlessAccelerationStructureDescriptorsAllocated);
+	case ERHIDescriptorType::Sampler:               return GET_STATID(STAT_VulkanBindlessSamplerDescriptorsAllocated);
+	case ERHIDescriptorType::TextureSRV:            return GET_STATID(STAT_VulkanBindlessSampledImageDescriptorsAllocated);
+	case ERHIDescriptorType::TextureUAV:            return GET_STATID(STAT_VulkanBindlessStorageImageDescriptorsAllocated);
+	case ERHIDescriptorType::TypedBufferSRV:        return GET_STATID(STAT_VulkanBindlessUniformTexelBufferDescriptorsAllocated);
+	case ERHIDescriptorType::TypedBufferUAV:        return GET_STATID(STAT_VulkanBindlessStorageTexelBufferDescriptorsAllocated);
+	case ERHIDescriptorType::BufferUAV:             return GET_STATID(STAT_VulkanBindlessStorageBufferDescriptorsAllocated);
+	case ERHIDescriptorType::CBV:                   return GET_STATID(STAT_VulkanBindlessUniformBufferDescriptorsAllocated);
+	case ERHIDescriptorType::BufferSRV:             return GET_STATID(STAT_VulkanBindlessUniformBufferDescriptorsAllocated);
+	case ERHIDescriptorType::AccelerationStructure: return GET_STATID(STAT_VulkanBindlessAccelerationStructureDescriptorsAllocated);
 	default: checkNoEntry();
 	}
 
 	return {};
 }
 
-static constexpr uint8 GetIndexForDescriptorType(VkDescriptorType DescriptorType)
-{
-	switch (DescriptorType)
-	{
-	case VK_DESCRIPTOR_TYPE_SAMPLER:                    return VulkanBindless::BindlessSamplerSet;
-	case VK_DESCRIPTOR_TYPE_SAMPLED_IMAGE:              return VulkanBindless::BindlessSampledImageSet;
-	case VK_DESCRIPTOR_TYPE_STORAGE_IMAGE:              return VulkanBindless::BindlessStorageImageSet;
-	case VK_DESCRIPTOR_TYPE_UNIFORM_TEXEL_BUFFER:       return VulkanBindless::BindlessUniformTexelBufferSet;
-	case VK_DESCRIPTOR_TYPE_STORAGE_TEXEL_BUFFER:       return VulkanBindless::BindlessStorageTexelBufferSet;
-	case VK_DESCRIPTOR_TYPE_STORAGE_BUFFER:             return VulkanBindless::BindlessStorageBufferSet;
-	case VK_DESCRIPTOR_TYPE_UNIFORM_BUFFER:             return VulkanBindless::BindlessUniformBufferSet;
-	case VK_DESCRIPTOR_TYPE_ACCELERATION_STRUCTURE_KHR: return VulkanBindless::BindlessAccelerationStructureSet;
-	default: checkNoEntry();
-	}
-
-	return VulkanBindless::MaxNumSets;
-}
-
-static constexpr uint8 GetIndexForDescriptorType(ERHIDescriptorType DescriptorType)
-{
-	switch (DescriptorType)
-	{
-	case ERHIDescriptorType::Sampler:               return VulkanBindless::BindlessSamplerSet;
-	case ERHIDescriptorType::TextureSRV:            return VulkanBindless::BindlessSampledImageSet;
-	case ERHIDescriptorType::TextureUAV:            return VulkanBindless::BindlessStorageImageSet;
-	case ERHIDescriptorType::TypedBufferSRV:        return VulkanBindless::BindlessUniformTexelBufferSet;
-	case ERHIDescriptorType::TypedBufferUAV:        return VulkanBindless::BindlessStorageTexelBufferSet;
-	case ERHIDescriptorType::BufferUAV:             return VulkanBindless::BindlessStorageBufferSet;
-	case ERHIDescriptorType::BufferSRV:             return VulkanBindless::BindlessUniformBufferSet;
-	case ERHIDescriptorType::CBV:                   return VulkanBindless::BindlessUniformBufferSet;
-	case ERHIDescriptorType::AccelerationStructure: return VulkanBindless::BindlessAccelerationStructureSet;
-	default: checkNoEntry();
-	}
-
-	return VulkanBindless::MaxNumSets;
-}
-
-
 static constexpr ERHIDescriptorTypeMask GetUEDescriptorTypeMaskFromVkDescriptorType(VkDescriptorType DescriptorType)
 {
 	switch (DescriptorType)
@@ -205,24 +118,21 @@ static inline uint32 GetDescriptorTypeSize(FVulkanDevice& Device, VkDescriptorTy
 			: DescriptorBufferProperties.storageBufferDescriptorSize;
 	case VK_DESCRIPTOR_TYPE_ACCELERATION_STRUCTURE_KHR:
 		return DescriptorBufferProperties.accelerationStructureDescriptorSize;
-	default: checkNoEntry();
+	case VK_DESCRIPTOR_TYPE_MUTABLE_EXT:
+	{
+		uint32 MutableDescriptorSize = FMath::Max3<uint32>(DescriptorBufferProperties.sampledImageDescriptorSize, DescriptorBufferProperties.storageImageDescriptorSize, DescriptorBufferProperties.accelerationStructureDescriptorSize);
+		if (bRobustBufferAccess)
+		{
+			MutableDescriptorSize = FMath::Max3<uint32>(MutableDescriptorSize, DescriptorBufferProperties.robustUniformTexelBufferDescriptorSize, DescriptorBufferProperties.robustStorageTexelBufferDescriptorSize);
+			MutableDescriptorSize = FMath::Max3<uint32>(MutableDescriptorSize, DescriptorBufferProperties.robustUniformBufferDescriptorSize, DescriptorBufferProperties.robustStorageBufferDescriptorSize);
+		}
+		else
+		{
+			MutableDescriptorSize = FMath::Max3<uint32>(MutableDescriptorSize, DescriptorBufferProperties.uniformTexelBufferDescriptorSize, DescriptorBufferProperties.storageTexelBufferDescriptorSize);
+			MutableDescriptorSize = FMath::Max3<uint32>(MutableDescriptorSize, DescriptorBufferProperties.uniformBufferDescriptorSize, DescriptorBufferProperties.storageBufferDescriptorSize);
 	}
-	return 0;
+		return MutableDescriptorSize;
 }
-
-
-static inline uint32 GetInitialDescriptorCount(VkDescriptorType DescriptorType)
-{
-	switch (DescriptorType)
-	{
-	case VK_DESCRIPTOR_TYPE_SAMPLER:                    return GVulkanBindlessMaxSamplerDescriptorCount;
-	case VK_DESCRIPTOR_TYPE_SAMPLED_IMAGE:              return GVulkanBindlessMaxSampledImageDescriptorCount;
-	case VK_DESCRIPTOR_TYPE_STORAGE_IMAGE:              return GVulkanBindlessMaxStorageImageDescriptorCount;
-	case VK_DESCRIPTOR_TYPE_UNIFORM_TEXEL_BUFFER:       return GVulkanBindlessMaxUniformTexelBufferDescriptorCount;
-	case VK_DESCRIPTOR_TYPE_STORAGE_TEXEL_BUFFER:       return GVulkanBindlessMaxStorageTexelBufferDescriptorCount;
-	case VK_DESCRIPTOR_TYPE_UNIFORM_BUFFER:             return GVulkanBindlessMaxUniformBufferDescriptorCount;
-	case VK_DESCRIPTOR_TYPE_STORAGE_BUFFER:             return GVulkanBindlessMaxStorageBufferDescriptorCount;
-	case VK_DESCRIPTOR_TYPE_ACCELERATION_STRUCTURE_KHR: return GVulkanBindlessMaxAccelerationStructureDescriptorCount;
 	default: checkNoEntry();
 	}
 	return 0;
@@ -253,18 +163,26 @@ static inline VkMemoryPropertyFlags GetDescriptorBufferMemoryType(FVulkanDevice&
 	return VK_MEMORY_PROPERTY_HOST_VISIBLE_BIT | VK_MEMORY_PROPERTY_HOST_COHERENT_BIT;
 }
 
-inline FVulkanBindlessDescriptorManager::BindlessSetState& FVulkanBindlessDescriptorManager::GetBindlessState(ERHIDescriptorType DescriptorType)
+inline FVulkanBindlessDescriptorManager::BindlessSetState& FVulkanBindlessDescriptorManager::GetBindlessState(ERHIDescriptorType RHIDescriptorType)
 {
-	const uint8 SetIndex = GetIndexForDescriptorType(DescriptorType);
-	check(SetIndex < VulkanBindless::NumBindlessSets);
-	return BindlessSetStates[SetIndex];
+	check(RHIDescriptorType != ERHIDescriptorType::Invalid);
+	const VkDescriptorType DescriptorType = GetVkDescriptorTypeFromUEDescriptorType(RHIDescriptorType);
+	return GetBindlessState(DescriptorType);
 }
 
 inline FVulkanBindlessDescriptorManager::BindlessSetState& FVulkanBindlessDescriptorManager::GetBindlessState(VkDescriptorType DescriptorType)
 {
-	const uint8 SetIndex = GetIndexForDescriptorType(DescriptorType);
-	check(SetIndex < VulkanBindless::NumBindlessSets);
-	return BindlessSetStates[SetIndex];
+	check(DescriptorType != VK_DESCRIPTOR_TYPE_MAX_ENUM);
+	for (BindlessSetState& State : BindlessSetStates)
+	{
+		if (State.DescriptorTypes.Contains(DescriptorType))
+		{
+			return State;
+		}
+	}
+	checkf(false, TEXT("Bindless state does not exist for descriptor type %s!"), VK_TYPE_TO_STRING(VkDescriptorType, DescriptorType));
+	checkNoEntry();
+	return BindlessSetStates[0];
 }
 
 // Check all the requirements to be running in Bindless using Descriptor Buffers
@@ -284,14 +202,15 @@ static bool VerifySupport(FVulkanDevice& Device, ERHIBindlessConfiguration Bindl
 		const bool bMeetsExtensionsRequirements =
 			OptionalDeviceExtensions.HasEXTDescriptorIndexing &&
 			OptionalDeviceExtensions.HasBufferDeviceAddress &&
-			OptionalDeviceExtensions.HasEXTDescriptorBuffer;
+			OptionalDeviceExtensions.HasEXTDescriptorBuffer &&
+			OptionalDeviceExtensions.HasEXTMutableDescriptorType;
 
 		if (bMeetsExtensionsRequirements)
 		{
 			const bool bMeetsPropertiesRequirements =
-				(GpuProps.limits.maxBoundDescriptorSets >= VulkanBindless::MaxNumSets) &&
-				(DescriptorBufferProperties.maxDescriptorBufferBindings >= VulkanBindless::MaxNumSets) &&
-				(DescriptorBufferProperties.maxResourceDescriptorBufferBindings >= VulkanBindless::NumBindlessSets) &&
+				(GpuProps.limits.maxBoundDescriptorSets >= 3) &&
+				(DescriptorBufferProperties.maxDescriptorBufferBindings >= 3) &&
+				(DescriptorBufferProperties.maxResourceDescriptorBufferBindings >= 1) &&
 				(DescriptorBufferProperties.maxSamplerDescriptorBufferBindings >= 1) &&
 				Device.GetDeviceMemoryManager().SupportsMemoryType(GetDescriptorBufferMemoryType(Device));
 
@@ -361,146 +280,129 @@ void FVulkanBindlessDescriptorManager::Init()
 
 	const VkBufferUsageFlags BufferUsageFlags = VK_BUFFER_USAGE_SHADER_DEVICE_ADDRESS_BIT | VK_BUFFER_USAGE_RESOURCE_DESCRIPTOR_BUFFER_BIT_EXT;
 
-	// Create the dummy layout for unsupported descriptor types
-	{
-		VkDescriptorSetLayoutCreateInfo EmptyDescriptorSetLayoutCreateInfo;
-		ZeroVulkanStruct(EmptyDescriptorSetLayoutCreateInfo, VK_STRUCTURE_TYPE_DESCRIPTOR_SET_LAYOUT_CREATE_INFO);
-		EmptyDescriptorSetLayoutCreateInfo.flags = VK_DESCRIPTOR_SET_LAYOUT_CREATE_DESCRIPTOR_BUFFER_BIT_EXT;
-		VERIFYVULKANRESULT(VulkanRHI::vkCreateDescriptorSetLayout(DeviceHandle, &EmptyDescriptorSetLayoutCreateInfo, VULKAN_CPU_ALLOCATOR, &EmptyDescriptorSetLayout));
-	}
-
 	{
-		auto InitBindlessSetState = [this](VkDescriptorType DescriptorType)
+		// Split up the descriptor types handled by each state
+		// 0 - Samplers
+		// 1 - All resources types (using mutable)
 		{
-			BindlessSetState& OutState = GetBindlessState(DescriptorType);
+			// Put samplers in state 0 and mutable in 
+			static_assert(UE_ARRAY_COUNT(BindlessSetStates) == 2, "BindlessSetStates array size has changed, must adjust decriptor type distribution!");
 
-			OutState.DescriptorTypes.Add(DescriptorType);
+			BindlessSetStates[0].DescriptorTypes.Add(VK_DESCRIPTOR_TYPE_SAMPLER);
 
-			OutState.DescriptorSize = FMath::Max(OutState.DescriptorSize, GetDescriptorTypeSize(Device, DescriptorType));
-			checkf((OutState.DescriptorSize > 0), TEXT("Descriptor Type [%s] returned an invalid descriptor size!"), VK_TYPE_TO_STRING(VkDescriptorType, DescriptorType));
+			BindlessSetStates[1].DescriptorTypes.Add(VK_DESCRIPTOR_TYPE_SAMPLED_IMAGE);
+			BindlessSetStates[1].DescriptorTypes.Add(VK_DESCRIPTOR_TYPE_STORAGE_IMAGE);
+			BindlessSetStates[1].DescriptorTypes.Add(VK_DESCRIPTOR_TYPE_UNIFORM_TEXEL_BUFFER);
+			BindlessSetStates[1].DescriptorTypes.Add(VK_DESCRIPTOR_TYPE_STORAGE_TEXEL_BUFFER);
+			BindlessSetStates[1].DescriptorTypes.Add(VK_DESCRIPTOR_TYPE_STORAGE_BUFFER);
+			BindlessSetStates[1].DescriptorTypes.Add(VK_DESCRIPTOR_TYPE_UNIFORM_BUFFER);
+			if (Device.GetOptionalExtensions().HasRaytracingExtensions())
+	{
+				BindlessSetStates[1].DescriptorTypes.Add(VK_DESCRIPTOR_TYPE_ACCELERATION_STRUCTURE_KHR);
+			}
 
-			OutState.DescriptorCapacity += GetInitialDescriptorCount(DescriptorType);
-			checkf((OutState.DescriptorCapacity > 0), TEXT("Descriptor Type [%s] returned an invalid descriptor count!"), VK_TYPE_TO_STRING(VkDescriptorType, DescriptorType));
-		};
+			// Add at the end so it can be excluded using Num-1 when only actual types are desired
+			BindlessSetStates[1].DescriptorTypes.Add(VK_DESCRIPTOR_TYPE_MUTABLE_EXT);
+		}
 
-		// Go through all the supported descriptor types in bindless
-		InitBindlessSetState(VK_DESCRIPTOR_TYPE_SAMPLER);
-		InitBindlessSetState(VK_DESCRIPTOR_TYPE_SAMPLED_IMAGE);
-		InitBindlessSetState(VK_DESCRIPTOR_TYPE_STORAGE_IMAGE);
-		InitBindlessSetState(VK_DESCRIPTOR_TYPE_UNIFORM_TEXEL_BUFFER);
-		InitBindlessSetState(VK_DESCRIPTOR_TYPE_STORAGE_TEXEL_BUFFER);
-		InitBindlessSetState(VK_DESCRIPTOR_TYPE_STORAGE_BUFFER);
-		InitBindlessSetState(VK_DESCRIPTOR_TYPE_UNIFORM_BUFFER);
-
-		// Create a dummy buffer for acceleration structures when they aren't supported (or ray tracing is disabled)
-		if (Device.GetOptionalExtensions().HasRaytracingExtensions())
+		// Setup the sampler states
+		BindlessSetState& SamplerState = GetBindlessState(VK_DESCRIPTOR_TYPE_SAMPLER);
 		{
-			InitBindlessSetState(VK_DESCRIPTOR_TYPE_ACCELERATION_STRUCTURE_KHR);
+			SamplerState.DescriptorSize = GetDescriptorTypeSize(Device, VK_DESCRIPTOR_TYPE_SAMPLER);
+			checkf((SamplerState.DescriptorSize > 0), TEXT("Descriptor Type [VK_DESCRIPTOR_TYPE_SAMPLER] returned an invalid descriptor size!"));
+			SamplerState.DescriptorCapacity = GVulkanBindlessMaxSamplerDescriptorCount;
+			checkf((SamplerState.DescriptorCapacity > 0), TEXT("Descriptor Type [VK_DESCRIPTOR_TYPE_SAMPLER] returned an invalid descriptor count!"));
+
+			const ERHIDescriptorTypeMask TypeMask = RHIDescriptorTypeMaskFromType(ERHIDescriptorType::Sampler);
+			TArray<TStatId, TInlineAllocator<1>> AllocatorStats;
+			AllocatorStats.Emplace(GET_STATID(STAT_BindlessSamplerDescriptorsAllocated));
+			SamplerState.Allocator = new FRHIHeapDescriptorAllocator(TypeMask, SamplerState.DescriptorCapacity, TConstArrayView<TStatId>(AllocatorStats));
 		}
-		else
+
+		// Setup the mutable descriptor states
+		BindlessSetState& MutableState = GetBindlessState(VK_DESCRIPTOR_TYPE_MUTABLE_EXT);
 		{
-			BindlessSetState& State = GetBindlessState(VK_DESCRIPTOR_TYPE_ACCELERATION_STRUCTURE_KHR);
-			State.DescriptorTypes.Add(VK_DESCRIPTOR_TYPE_STORAGE_BUFFER);
-			State.DescriptorSize = GetDescriptorTypeSize(Device, VK_DESCRIPTOR_TYPE_STORAGE_BUFFER);
-			State.DescriptorCapacity = 16;
+			MutableState.DescriptorSize = GetDescriptorTypeSize(Device, VK_DESCRIPTOR_TYPE_MUTABLE_EXT);
+			checkf((MutableState.DescriptorSize > 0), TEXT("Descriptor Type [VK_DESCRIPTOR_TYPE_MUTABLE_EXT] returned an invalid descriptor size!"));
+			MutableState.DescriptorCapacity = GVulkanBindlessMaxResourceDescriptorCount;
+			checkf((MutableState.DescriptorCapacity > 0), TEXT("Descriptor Type [VK_DESCRIPTOR_TYPE_MUTABLE_EXT] returned an invalid descriptor count!"));
+
+			ERHIDescriptorTypeMask TypeMask{};
+			TArray<TStatId, TInlineAllocator<1>> AllocatorStats;
+			AllocatorStats.Emplace(GET_STATID(STAT_BindlessResourceDescriptorsAllocated));
+			for (int32 TypeIndex = 0; TypeIndex < MutableState.DescriptorTypes.Num() - 1; ++TypeIndex)
+			{
+				const VkDescriptorType DescriptorType = MutableState.DescriptorTypes[TypeIndex];
+				EnumAddFlags(TypeMask, GetUEDescriptorTypeMaskFromVkDescriptorType(DescriptorType));
+			}
+			MutableState.Allocator = new FRHIHeapDescriptorAllocator(TypeMask, MutableState.DescriptorCapacity, TConstArrayView<TStatId>(AllocatorStats));
 		}
 
+
 		// Fill the DescriptorSetLayout for a BindlessSetState
 		auto CreateDescriptorSetLayout = [&, &DeviceRef=Device](const BindlessSetState& State)
 		{
-			if (State.DescriptorTypes.Num() == 0)
+			const bool IsSamplerSet = (State.DescriptorTypes[0] == VK_DESCRIPTOR_TYPE_SAMPLER);
+
+			const VkDescriptorSetLayoutBinding DescriptorSetLayoutBindings =
 			{
-				return EmptyDescriptorSetLayout;
-			}
-			else
+				.binding = 0,
+				.descriptorType = IsSamplerSet ? VK_DESCRIPTOR_TYPE_SAMPLER : VK_DESCRIPTOR_TYPE_MUTABLE_EXT,
+				.descriptorCount = (uint32)(IsSamplerSet ? GVulkanBindlessMaxSamplerDescriptorCount : GVulkanBindlessMaxResourceDescriptorCount),
+				.stageFlags = VK_SHADER_STAGE_ALL,
+				.pImmutableSamplers = nullptr
+			};
+
+			const VkDescriptorBindingFlags DescriptorBindingFlags = 0; //NOTE: Eventually change to VK_DESCRIPTOR_BINDING_VARIABLE_DESCRIPTOR_COUNT_BIT
+			const VkDescriptorSetLayoutBindingFlagsCreateInfo DescriptorSetLayoutBindingFlagsCreateInfo =
 			{
-				VkDescriptorSetLayoutBinding DescriptorSetLayoutBinding;
-				DescriptorSetLayoutBinding.binding = 0;
-				DescriptorSetLayoutBinding.descriptorCount = State.DescriptorCapacity;  // todo-jn: resizable
-				DescriptorSetLayoutBinding.stageFlags = VK_SHADER_STAGE_ALL;
-				DescriptorSetLayoutBinding.pImmutableSamplers = nullptr; // todo-jn: ImmutableSamplers
-
-				// These flags are implied with descriptor_buffer: VK_DESCRIPTOR_BINDING_UPDATE_AFTER_BIND_BIT,VK_DESCRIPTOR_BINDING_PARTIALLY_BOUND_BIT,VK_DESCRIPTOR_BINDING_UPDATE_UNUSED_WHILE_PENDING_BIT
-				// :todo-jn: add support for VK_DESCRIPTOR_BINDING_VARIABLE_DESCRIPTOR_COUNT_BIT when drivers are fixed to allow for buffers to grow
-				const VkDescriptorBindingFlags DescriptorBindingFlags = 0; // VK_DESCRIPTOR_BINDING_VARIABLE_DESCRIPTOR_COUNT_BIT;
-
-				VkDescriptorSetLayoutBindingFlagsCreateInfo DescriptorSetLayoutBindingFlagsCreateInfo;
-				ZeroVulkanStruct(DescriptorSetLayoutBindingFlagsCreateInfo, VK_STRUCTURE_TYPE_DESCRIPTOR_SET_LAYOUT_BINDING_FLAGS_CREATE_INFO);
-				DescriptorSetLayoutBindingFlagsCreateInfo.bindingCount = 1;
-				DescriptorSetLayoutBindingFlagsCreateInfo.pBindingFlags = &DescriptorBindingFlags;
-
-				VkDescriptorSetLayoutCreateInfo DescriptorSetLayoutCreateInfo;
-				ZeroVulkanStruct(DescriptorSetLayoutCreateInfo, VK_STRUCTURE_TYPE_DESCRIPTOR_SET_LAYOUT_CREATE_INFO);
-				DescriptorSetLayoutCreateInfo.pBindings = &DescriptorSetLayoutBinding;
-				DescriptorSetLayoutCreateInfo.bindingCount = 1;
-				DescriptorSetLayoutCreateInfo.flags = VK_DESCRIPTOR_SET_LAYOUT_CREATE_DESCRIPTOR_BUFFER_BIT_EXT;
-				DescriptorSetLayoutCreateInfo.pNext = &DescriptorSetLayoutBindingFlagsCreateInfo;
-
-				VkMutableDescriptorTypeCreateInfoEXT MutableDescriptorTypeCreateInfo;
-				VkMutableDescriptorTypeListEXT MutableDescriptorTypeList;
-				if (State.DescriptorTypes.Num() == 1)
-				{
-					DescriptorSetLayoutBinding.descriptorType = State.DescriptorTypes[0];
-				}
-				else
-				{
-					DescriptorSetLayoutBinding.descriptorType = VK_DESCRIPTOR_TYPE_MUTABLE_EXT;
+				.sType = VK_STRUCTURE_TYPE_DESCRIPTOR_SET_LAYOUT_BINDING_FLAGS_CREATE_INFO,
+				.pNext = nullptr,
+				.bindingCount = 1,
+				.pBindingFlags = &DescriptorBindingFlags
+			};
 
-					MutableDescriptorTypeList.descriptorTypeCount = State.DescriptorTypes.Num();
-					MutableDescriptorTypeList.pDescriptorTypes = State.DescriptorTypes.GetData();
+			const uint32 MutableTypeCount = MutableState.DescriptorTypes.Num() - 1;
+			const VkMutableDescriptorTypeListEXT MutableDescriptorTypeList =
+			{
+				.descriptorTypeCount = MutableTypeCount,
+				.pDescriptorTypes = MutableState.DescriptorTypes.GetData()
+			};
 
-					ZeroVulkanStruct(MutableDescriptorTypeCreateInfo, VK_STRUCTURE_TYPE_MUTABLE_DESCRIPTOR_TYPE_CREATE_INFO_EXT);
-					MutableDescriptorTypeCreateInfo.mutableDescriptorTypeListCount = 1;
-					MutableDescriptorTypeCreateInfo.pMutableDescriptorTypeLists = &MutableDescriptorTypeList;
+			const VkMutableDescriptorTypeCreateInfoEXT MutableDescriptorTypeCreateInfo =
+			{
+				.sType = VK_STRUCTURE_TYPE_MUTABLE_DESCRIPTOR_TYPE_CREATE_INFO_EXT,
+				.pNext = &DescriptorSetLayoutBindingFlagsCreateInfo,
+				.mutableDescriptorTypeListCount = 1,
+				.pMutableDescriptorTypeLists = &MutableDescriptorTypeList
+			};
 
-					MutableDescriptorTypeCreateInfo.pNext = DescriptorSetLayoutCreateInfo.pNext;
-					DescriptorSetLayoutCreateInfo.pNext = &MutableDescriptorTypeCreateInfo;
-				}
+			// NOTE: Eventually add VkDescriptorSetLayoutBindingFlagsCreateInfo to pNext for VK_DESCRIPTOR_BINDING_VARIABLE_DESCRIPTOR_COUNT_BIT
+			const VkDescriptorSetLayoutCreateInfo DescriptorSetLayoutCreateInfo =
+				{
+				.sType = VK_STRUCTURE_TYPE_DESCRIPTOR_SET_LAYOUT_CREATE_INFO,
+				.pNext = IsSamplerSet ? (void*)&DescriptorSetLayoutBindingFlagsCreateInfo : (void*)&MutableDescriptorTypeCreateInfo,
+				.flags = VK_DESCRIPTOR_SET_LAYOUT_CREATE_DESCRIPTOR_BUFFER_BIT_EXT,
+				.bindingCount = 1,
+				.pBindings = &DescriptorSetLayoutBindings
+			};
 
 				VkDescriptorSetLayout DescriptorSetLayout = VK_NULL_HANDLE;
 				VERIFYVULKANRESULT(VulkanRHI::vkCreateDescriptorSetLayout(DeviceHandle, &DescriptorSetLayoutCreateInfo, VULKAN_CPU_ALLOCATOR, &DescriptorSetLayout));
+			VULKAN_SET_DEBUG_NAME(Device, VK_OBJECT_TYPE_DESCRIPTOR_SET_LAYOUT, DescriptorSetLayout, TEXT("BindlessDescriptorSet(%s)"), IsSamplerSet ? TEXT("Sampler") : TEXT("Resources"));
 
+			// Safe check
 				VkDeviceSize BindingOffset = 0;
 				VulkanRHI::vkGetDescriptorSetLayoutBindingOffsetEXT(DeviceHandle, DescriptorSetLayout, 0u, &BindingOffset);
 				check(BindingOffset == 0);
 
-				VULKAN_SET_DEBUG_NAME(DeviceRef, VK_OBJECT_TYPE_DESCRIPTOR_SET_LAYOUT, DescriptorSetLayout, TEXT("BindlessArray(%s)"), VK_TYPE_TO_STRING(VkDescriptorType, DescriptorSetLayoutBinding.descriptorType));
-
 				return DescriptorSetLayout;
-			}
 		};
 
+
 		// Create the descriptor buffer for a BindlessSetState
 		auto CreateDescriptorBuffer = [&](BindlessSetState& InOutState, VkDescriptorBufferBindingInfoEXT& OutBufferBindingInfo)
 		{
-			// Skip unsupported descriptors
-			if (InOutState.DescriptorTypes.Num() == 0)
-			{
-				return 0u;
-			}
-
-			ERHIDescriptorTypeMask TypeMask{};
-			TArray<TStatId, TInlineAllocator<3>> AllocatorStats;
-
-			AllocatorStats.Emplace(GET_STATID(STAT_VulkanBindlessDescriptorsAllocated));
-
-			for (VkDescriptorType DescriptorType : InOutState.DescriptorTypes)
-			{
-				EnumAddFlags(TypeMask, GetUEDescriptorTypeMaskFromVkDescriptorType(DescriptorType));
-				AllocatorStats.Emplace(GetStatForDescriptorType(DescriptorType));
-
-				if (DescriptorType == VK_DESCRIPTOR_TYPE_SAMPLER)
-				{
-					AllocatorStats.Emplace(GET_STATID(STAT_BindlessSamplerDescriptorsAllocated));
-				}
-				else
-				{
-					AllocatorStats.Emplace(GET_STATID(STAT_BindlessResourceDescriptorsAllocated));
-				}
-			}
-
-			InOutState.Allocator = new FRHIHeapDescriptorAllocator(TypeMask, InOutState.DescriptorCapacity, TConstArrayView<TStatId>(AllocatorStats));
-
 			const bool IsSamplerSet = (InOutState.DescriptorTypes[0] == VK_DESCRIPTOR_TYPE_SAMPLER);
 			const VkBufferUsageFlags BufferUsageFlags =
 				VK_BUFFER_USAGE_SHADER_DEVICE_ADDRESS_BIT |
@@ -522,6 +424,12 @@ void FVulkanBindlessDescriptorManager::Init()
 					TEXT("Sampler descriptor buffer size [%u] exceeded maximum [%llu]."),
 					DescriptorBufferSize, DescriptorBufferProperties.samplerDescriptorBufferAddressSpaceSize);
 			}
+			else
+			{
+				checkf(DescriptorBufferSize < DescriptorBufferProperties.resourceDescriptorBufferAddressSpaceSize,
+					TEXT("Combined resource descriptor buffer size of [%u] exceeded maximum [%llu]."),
+					DescriptorBufferSize, DescriptorBufferProperties.resourceDescriptorBufferAddressSpaceSize);
+			}
 
 			// Create descriptor buffer
 			{
@@ -567,17 +475,14 @@ void FVulkanBindlessDescriptorManager::Init()
 
 			// Add the types to the global groups
 			GRHIGlobals.DescriptorTypeGroups.Emplace(InOutState.Allocator->GetTypeMask());
-
-			return IsSamplerSet ? 0u : DescriptorBufferSize;
 		};
 
 		// Fill in one state for each descriptor type
-		uint32 TotalResourceDescriptorBufferSize = 0;
 		for (uint32 SetIndex = 0; SetIndex < VulkanBindless::NumBindlessSets; ++SetIndex)
 		{
 			BindlessSetState& State = BindlessSetStates[SetIndex];
 			State.DescriptorSetLayout = CreateDescriptorSetLayout(State);
-			TotalResourceDescriptorBufferSize += CreateDescriptorBuffer(State, BufferBindingInfo[SetIndex]);
+			CreateDescriptorBuffer(State, BufferBindingInfo[SetIndex]);
 		}
 
 		// Fill in the state for single-use UB
@@ -607,10 +512,6 @@ void FVulkanBindlessDescriptorManager::Init()
 			VERIFYVULKANRESULT(VulkanRHI::vkCreateDescriptorSetLayout(DeviceHandle, &DescriptorSetLayoutCreateInfo, VULKAN_CPU_ALLOCATOR, &SingleUseUBDescriptorSetLayout));
 			VULKAN_SET_DEBUG_NAME(Device, VK_OBJECT_TYPE_DESCRIPTOR_SET_LAYOUT, SingleUseUBDescriptorSetLayout, TEXT("UBPushDescriptors(%d)"), NumTotalBindings);
 		}
-
-		checkf(TotalResourceDescriptorBufferSize < DescriptorBufferProperties.resourceDescriptorBufferAddressSpaceSize,
-			TEXT("Combined resource descriptor buffer size of [%u] exceeded maximum [%llu]."),
-			TotalResourceDescriptorBufferSize, DescriptorBufferProperties.resourceDescriptorBufferAddressSpaceSize);
 	}
 
 	// Now create the single pipeline layout used by everything
@@ -623,12 +524,16 @@ void FVulkanBindlessDescriptorManager::Init()
 		BindlessDescriptorSetLayouts[VulkanBindless::BindlessSingleUseUniformBufferSet] = SingleUseUBDescriptorSetLayout;
 		check(VulkanBindless::BindlessSingleUseUniformBufferSet < BindlessDescriptorSetLayouts.Num());
 
-		VkPipelineLayoutCreateInfo PipelineLayoutCreateInfo;
-		ZeroVulkanStruct(PipelineLayoutCreateInfo, VK_STRUCTURE_TYPE_PIPELINE_LAYOUT_CREATE_INFO);
-		PipelineLayoutCreateInfo.setLayoutCount = BindlessDescriptorSetLayouts.Num();
-		PipelineLayoutCreateInfo.pSetLayouts = BindlessDescriptorSetLayouts.GetData();
-		PipelineLayoutCreateInfo.pushConstantRangeCount = UE_ARRAY_COUNT(PushConstantRanges);
-		PipelineLayoutCreateInfo.pPushConstantRanges = PushConstantRanges;
+		const VkPipelineLayoutCreateInfo PipelineLayoutCreateInfo =
+		{
+			.sType = VK_STRUCTURE_TYPE_PIPELINE_LAYOUT_CREATE_INFO,
+			.pNext = nullptr,
+			.flags = 0,
+			.setLayoutCount = (uint32)BindlessDescriptorSetLayouts.Num(),
+			.pSetLayouts = BindlessDescriptorSetLayouts.GetData(),
+			.pushConstantRangeCount = UE_ARRAY_COUNT(PushConstantRanges),
+			.pPushConstantRanges = PushConstantRanges
+		};
 
 		checkf((uint32)EPushConstantRangeType::Count == UE_ARRAY_COUNT(PushConstantRanges),
 			TEXT("Mismatch between expected push constant range count (%d) and provided array size (%d)"), 
@@ -659,6 +564,12 @@ void FVulkanBindlessDescriptorManager::Deinit()
 				VulkanRHI::vkUnmapMemory(DeviceHandle, State.MemoryHandle);
 				VulkanRHI::vkFreeMemory(DeviceHandle, State.MemoryHandle, VULKAN_CPU_ALLOCATOR);
 				State.MemoryHandle = VK_NULL_HANDLE;
+
+				if (State.Allocator)
+				{
+					delete State.Allocator;
+					State.Allocator = nullptr;
+				}
 			};
 
 		for (uint32 SetIndex = 0; SetIndex < VulkanBindless::NumBindlessSets; ++SetIndex)
@@ -672,9 +583,6 @@ void FVulkanBindlessDescriptorManager::Deinit()
 
 		VulkanRHI::vkDestroyDescriptorSetLayout(DeviceHandle, SingleUseUBDescriptorSetLayout, VULKAN_CPU_ALLOCATOR);
 		SingleUseUBDescriptorSetLayout = VK_NULL_HANDLE;
-
-		VulkanRHI::vkDestroyDescriptorSetLayout(DeviceHandle, EmptyDescriptorSetLayout, VULKAN_CPU_ALLOCATOR);
-		EmptyDescriptorSetLayout = VK_NULL_HANDLE;
 	}
 }
 
@@ -723,6 +631,7 @@ FRHIDescriptorHandle FVulkanBindlessDescriptorManager::AllocateDescriptor(ERHIDe
 {
 	if (IsSupported())
 	{
+		INC_DWORD_STAT_BY_FName(GetStatForDescriptorType(DescriptorType).GetName(), 1);
 		return GetBindlessState(DescriptorType).Allocator->Allocate(DescriptorType);
 	}
 
@@ -733,6 +642,7 @@ TOptional<FRHIDescriptorAllocation> FVulkanBindlessDescriptorManager::AllocateDe
 {
 	if (IsSupported())
 	{
+		INC_DWORD_STAT_BY_FName(GetStatForDescriptorType(DescriptorType).GetName(), DescriptorCount);
 		return GetBindlessState(DescriptorType).Allocator->Allocate(DescriptorCount);
 	}
 
@@ -780,6 +690,8 @@ void FVulkanBindlessDescriptorManager::FreeDescriptor(FRHIDescriptorHandle Descr
 {
 	if (DescriptorHandle.IsValid())
 	{
+		DEC_DWORD_STAT_BY_FName(GetStatForDescriptorType(DescriptorHandle.GetType()).GetName(), 1);
+
 		BindlessSetState& State = GetBindlessState(DescriptorHandle.GetType());
 		State.Allocator->Free(DescriptorHandle);
 
@@ -788,10 +700,12 @@ void FVulkanBindlessDescriptorManager::FreeDescriptor(FRHIDescriptorHandle Descr
 	}
 }
 
-void FVulkanBindlessDescriptorManager::FreeDescriptors(FRHIDescriptorAllocation Allocation, VkDescriptorType DescriptorType)
+void FVulkanBindlessDescriptorManager::FreeDescriptors(FRHIDescriptorAllocation Allocation, ERHIDescriptorType DescriptorType)
 {
 	if (Allocation.Count > 0)
 	{
+		DEC_DWORD_STAT_BY_FName(GetStatForDescriptorType(DescriptorType).GetName(), Allocation.Count);
+
 		BindlessSetState& State = GetBindlessState(DescriptorType);
 		State.Allocator->Free(Allocation);
 
@@ -816,7 +730,7 @@ void FVulkanBindlessDescriptorManager::UpdateDescriptorInternal(const FVulkanCon
 	ZeroVulkanStruct(Info, VK_STRUCTURE_TYPE_DESCRIPTOR_GET_INFO_EXT);
 	Info.type = DescriptorType;
 	Info.data = DescriptorData;
-	VulkanRHI::vkGetDescriptorEXT(Device.GetHandle(), &Info, State.DescriptorSize, CPUDescriptor);
+	VulkanRHI::vkGetDescriptorEXT(Device.GetHandle(), &Info, GetDescriptorTypeSize(Device, DescriptorType), CPUDescriptor);
 
 	if (bImmediateUpdate)
 	{
diff --git a/Engine/Source/Runtime/VulkanRHI/Private/VulkanBindlessDescriptorManager.h b/Engine/Source/Runtime/VulkanRHI/Private/VulkanBindlessDescriptorManager.h
--- a/Engine/Source/Runtime/VulkanRHI/Private/VulkanBindlessDescriptorManager.h
+++ b/Engine/Source/Runtime/VulkanRHI/Private/VulkanBindlessDescriptorManager.h
@@ -41,7 +41,7 @@ public:
 	FRHIDescriptorHandle AllocateDescriptor(VkSampler VulkanSampler);
 
 	void FreeDescriptor(FRHIDescriptorHandle DescriptorHandle);
-	void FreeDescriptors(FRHIDescriptorAllocation Allocation, VkDescriptorType DescriptorType);
+	void FreeDescriptors(FRHIDescriptorAllocation Allocation, ERHIDescriptorType DescriptorType);
 
 	void UpdateDescriptor(const FVulkanContextArray& Contexts, FRHIDescriptorHandle DescriptorHandle, VkSampler VulkanSampler);
 	void UpdateDescriptor(const FVulkanContextArray& Contexts, FRHIDescriptorHandle DescriptorHandle, VkImageView VulkanImage, bool bIsDepthStencil, bool bImmediateUpdate = true);
@@ -86,8 +86,6 @@ private:
 	ERHIBindlessConfiguration Configuration{};
 	ERHIPipeline              ConfiguredPipelines = ERHIPipeline::None;
 
-	VkDescriptorSetLayout EmptyDescriptorSetLayout = VK_NULL_HANDLE;
-
 	struct BindlessSetState
 	{
 		~BindlessSetState()
@@ -95,6 +93,13 @@ private:
 			if (CPUDescriptorMemory)
 			{
 				FMemory::Free(CPUDescriptorMemory);
+				CPUDescriptorMemory = nullptr;
+			}
+
+			if (Allocator)
+			{
+				delete Allocator;
+				Allocator = nullptr;
 			}
 		}
 
@@ -123,7 +128,7 @@ private:
 	TStaticArray<VkDescriptorSetLayout, VulkanBindless::MaxNumSets> BindlessDescriptorSetLayouts;
 	static const VkPushConstantRange PushConstantRanges[(uint8)EPushConstantRangeType::Count];
 
-	BindlessSetState& GetBindlessState(ERHIDescriptorType DescriptorType);
+	BindlessSetState& GetBindlessState(ERHIDescriptorType RHIDescriptorType);
 	BindlessSetState& GetBindlessState(VkDescriptorType DescriptorType);
 
 	void UpdateDescriptorInternal(const FVulkanContextArray& Contexts, FRHIDescriptorHandle DescriptorHandle, VkDescriptorDataEXT DescriptorData, bool bImmediateUpdate);
diff --git a/Engine/Source/Runtime/VulkanRHI/Public/VulkanCommon.h b/Engine/Source/Runtime/VulkanRHI/Public/VulkanCommon.h
--- a/Engine/Source/Runtime/VulkanRHI/Public/VulkanCommon.h
+++ b/Engine/Source/Runtime/VulkanRHI/Public/VulkanCommon.h
@@ -18,13 +18,6 @@
 // ONLY used for debugging binding table/descriptor set bugs/mismatches.
 #define VULKAN_ENABLE_BINDING_DEBUG_NAMES						0
 
-// Bindless uses one descriptor set per resource type plus one for single use UBs, for a total of 9.
-// On platforms that have a maximum below 9, use mutable descriptors to lower total count to 6 by
-// creating a descriptor set per resource type rather than per descriptor type.  On most platforms,
-// these descriptors have the same size and should still be packed tightly.
-// NOTE: Should invalidate Vulkan shaders when changing this value (for example, update VulkanCommon GUID).
-#define VULKAN_REDUCE_BINDLESS_DESCRIPTOR_SET_COUNT				0
-
 namespace ShaderStage
 {
 	// There should be one value for each value in EShaderFrequency.
@@ -108,27 +101,16 @@ namespace VulkanBindless
 	{
 		BindlessSamplerSet = 0,
 
-#if VULKAN_REDUCE_BINDLESS_DESCRIPTOR_SET_COUNT
 		BindlessStorageBufferSet = 1,
 		BindlessUniformBufferSet = 1,
 
-		BindlessStorageImageSet = 2,
-		BindlessSampledImageSet = 2,
-
-		BindlessStorageTexelBufferSet = 3,
-		BindlessUniformTexelBufferSet = 3,
-#else
-		BindlessStorageBufferSet,
-		BindlessUniformBufferSet,
+		BindlessStorageImageSet = 1,
+		BindlessSampledImageSet = 1,
 
-		BindlessStorageImageSet,
-		BindlessSampledImageSet,
-
-		BindlessStorageTexelBufferSet,
-		BindlessUniformTexelBufferSet,
-#endif
+		BindlessStorageTexelBufferSet = 1,
+		BindlessUniformTexelBufferSet = 1,
 
-		BindlessAccelerationStructureSet,
+		BindlessAccelerationStructureSet = 1,
 
 		// Number of sets reserved for samplers/resources
 		NumBindlessSets,
