PR #14273: Adding support for Vulkan End of Frame markers

#rnx
#Synced-CL 50399050

diff --git a/Engine/Source/Runtime/VulkanRHI/Private/VulkanContext.cpp b/Engine/Source/Runtime/VulkanRHI/Private/VulkanContext.cpp
--- a/Engine/Source/Runtime/VulkanRHI/Private/VulkanContext.cpp
+++ b/Engine/Source/Runtime/VulkanRHI/Private/VulkanContext.cpp
@@ -724,7 +724,27 @@ void FVulkanCommandListContext::ApplyShaderBindingLayout(VkShaderStageFlags Shad
 	}
 }
 
-
+// Inserts a Vulkan-specific end frame marker, which can be used by tools like RenderDoc to
+// identify when a frame is done when there is no native Present (e.g. with a standalone XR
+// device rendering to an offscreen buffer
+void FVulkanCommandListContext::InsertEndFrameDebugLabel()
+{
+#if VULKAN_ENABLE_DRAW_MARKERS
+	if (auto CmdDbgInsertLabel = Device.GetCmdInsertDebugLabel())
+	{
+		FTCHARToUTF8 Converter(TEXT("vr-marker,frame_end,type,application"));
+		VkDebugUtilsLabelEXT Label;
+		ZeroVulkanStruct(Label, VK_STRUCTURE_TYPE_DEBUG_UTILS_LABEL_EXT);
+		Label.pLabelName = Converter.Get();
+		const FLinearColor LColor = FLinearColor::Blue;
+		Label.color[0] = LColor.R;
+		Label.color[1] = LColor.G;
+		Label.color[2] = LColor.B;
+		Label.color[3] = LColor.A;
+		CmdDbgInsertLabel(GetCommandBuffer().GetHandle(), &Label);
+	}
+#endif //VULKAN_ENABLE_DRAW_MARKERS
+}
 
 
 TLockFreePointerListUnordered<FVulkanUploadContext, PLATFORM_CACHE_LINE_SIZE> FVulkanUploadContext::Pool;
diff --git a/Engine/Source/Runtime/VulkanRHI/Private/VulkanContext.h b/Engine/Source/Runtime/VulkanRHI/Private/VulkanContext.h
--- a/Engine/Source/Runtime/VulkanRHI/Private/VulkanContext.h
+++ b/Engine/Source/Runtime/VulkanRHI/Private/VulkanContext.h
@@ -288,7 +288,6 @@ public:
 	virtual void RHIBeginBreadcrumbGPU(FRHIBreadcrumbNode* Breadcrumb) final override;
 	virtual void RHIEndBreadcrumbGPU  (FRHIBreadcrumbNode* Breadcrumb) final override;
 #endif
-
 	virtual void RHISetComputePipelineState(FRHIComputePipelineState* ComputePipelineState) final override;
 	virtual void RHIDispatchComputeShader(uint32 ThreadGroupCountX, uint32 ThreadGroupCountY, uint32 ThreadGroupCountZ) final override;
 	virtual void RHIDispatchIndirectComputeShader(FRHIBuffer* ArgumentBuffer, uint32 ArgumentOffset) final override;
@@ -343,6 +342,8 @@ public:
 		uint32 NumBindings, const FRayTracingLocalShaderBindings* Bindings,
 		ERayTracingBindingType BindingType) final override;
 
+	void InsertEndFrameDebugLabel();
+
 	inline FVulkanPendingGfxState* GetPendingGfxState()
 	{
 		return PendingGfxState;
diff --git a/Engine/Source/Runtime/VulkanRHI/Private/VulkanDevice.cpp b/Engine/Source/Runtime/VulkanRHI/Private/VulkanDevice.cpp
--- a/Engine/Source/Runtime/VulkanRHI/Private/VulkanDevice.cpp
+++ b/Engine/Source/Runtime/VulkanRHI/Private/VulkanDevice.cpp
@@ -652,6 +652,7 @@ void FVulkanDevice::SetupDrawMarkers()
 			{
 				DebugMarkers.CmdBeginDebugLabel = (PFN_vkCmdBeginDebugUtilsLabelEXT)(void*)VulkanRHI::vkGetInstanceProcAddr(RHI->GetInstance(), "vkCmdBeginDebugUtilsLabelEXT");
 				DebugMarkers.CmdEndDebugLabel = (PFN_vkCmdEndDebugUtilsLabelEXT)(void*)VulkanRHI::vkGetInstanceProcAddr(RHI->GetInstance(), "vkCmdEndDebugUtilsLabelEXT");
+				DebugMarkers.CmdInsertLabel = (PFN_vkCmdInsertDebugUtilsLabelEXT)(void*)VulkanRHI::vkGetInstanceProcAddr(RHI->GetInstance(), "vkCmdInsertDebugUtilsLabelEXT");
 			}
 			if (bUseName)
 			{
diff --git a/Engine/Source/Runtime/VulkanRHI/Private/VulkanDevice.h b/Engine/Source/Runtime/VulkanRHI/Private/VulkanDevice.h
--- a/Engine/Source/Runtime/VulkanRHI/Private/VulkanDevice.h
+++ b/Engine/Source/Runtime/VulkanRHI/Private/VulkanDevice.h
@@ -519,6 +519,11 @@ public:
 	void NotifyDeletedImage(VkImage Image, bool bRenderTarget);
 
 #if VULKAN_ENABLE_DRAW_MARKERS
+	inline PFN_vkCmdInsertDebugUtilsLabelEXT GetCmdInsertDebugLabel() const
+	{
+		return DebugMarkers.CmdInsertLabel;
+	}
+
 	PFN_vkCmdBeginDebugUtilsLabelEXT GetCmdBeginDebugLabel() const
 	{
 		return DebugMarkers.CmdBeginDebugLabel;
@@ -730,6 +735,7 @@ private:
 	struct
 	{
 		PFN_vkSetDebugUtilsObjectNameEXT	SetDebugName = nullptr;
+		PFN_vkCmdInsertDebugUtilsLabelEXT	CmdInsertLabel = nullptr;
 		PFN_vkCmdBeginDebugUtilsLabelEXT	CmdBeginDebugLabel = nullptr;
 		PFN_vkCmdEndDebugUtilsLabelEXT		CmdEndDebugLabel = nullptr;
 	} DebugMarkers;
diff --git a/Engine/Source/Runtime/VulkanRHI/Private/VulkanViewport.cpp b/Engine/Source/Runtime/VulkanRHI/Private/VulkanViewport.cpp
--- a/Engine/Source/Runtime/VulkanRHI/Private/VulkanViewport.cpp
+++ b/Engine/Source/Runtime/VulkanRHI/Private/VulkanViewport.cpp
@@ -755,6 +755,17 @@ bool FVulkanViewport::Present(FVulkanCommandListContext& Context, const FRHIPres
 		}
 	}
 
+	const bool bHasCustomPresent = IsValidRef(CustomPresent);
+
+	if (bHasCustomPresent)
+	{
+		// If we don't do the actual Present call (e.g. XR device rendering to an offscreen
+		// compositor buffer), we need to manually set the end of the frame marker
+		// for tools like RenderDoc to be able to judge when the frame has ended before we submit
+		// the active command buffer below
+		Context.InsertEndFrameDebugLabel();
+	}
+
 	// Submit any accumulated commands or syncs, wait until they hit the queue so that we can present
 	Context.FlushCommands(EVulkanFlushFlags::WaitForSubmission);
 
@@ -762,7 +773,6 @@ bool FVulkanViewport::Present(FVulkanCommandListContext& Context, const FRHIPres
 	int32 SyncInterval = 0;
 	bool bNeedNativePresent = true;
 
-	const bool bHasCustomPresent = IsValidRef(CustomPresent);
 	if (bHasCustomPresent)
 	{
 		SCOPE_CYCLE_COUNTER(STAT_VulkanCustomPresentTime);
