VulkanRHI: Fix mip index for readback textures copies

#rnx

#Synced-CL 50399050

diff --git a/Engine/Source/Runtime/VulkanRHI/Private/VulkanTexture.cpp b/Engine/Source/Runtime/VulkanRHI/Private/VulkanTexture.cpp
--- a/Engine/Source/Runtime/VulkanRHI/Private/VulkanTexture.cpp
+++ b/Engine/Source/Runtime/VulkanRHI/Private/VulkanTexture.cpp
@@ -2620,7 +2620,7 @@ void FVulkanCommandListContext::RHICopyTexture(FRHITexture* SourceTexture, FRHIT
 			CopyRegion[Index].bufferRowLength = Size.X;
 			CopyRegion[Index].bufferImageHeight = Size.Y;
 			CopyRegion[Index].imageSubresource.aspectMask = Source->GetFullAspectMask();
-			CopyRegion[Index].imageSubresource.mipLevel = SourceMipIndex;
+			CopyRegion[Index].imageSubresource.mipLevel = SourceMipIndex + Index;
 			CopyRegion[Index].imageSubresource.baseArrayLayer = SourceSliceIndex;
 			CopyRegion[Index].imageSubresource.layerCount = 1;
 			CopyRegion[Index].imageOffset.x = CopyInfo.SourcePosition.X;
