[HWRT] Vulkan SBT needs to be cleared to 0 to match the behavior of other platforms

#rnx

#Synced-CL 50366487

#Synced-CL 50366487

diff --git a/Engine/Source/Runtime/VulkanRHI/Private/VulkanRayTracing.cpp b/Engine/Source/Runtime/VulkanRHI/Private/VulkanRayTracing.cpp
--- a/Engine/Source/Runtime/VulkanRHI/Private/VulkanRayTracing.cpp
+++ b/Engine/Source/Runtime/VulkanRHI/Private/VulkanRayTracing.cpp
@@ -846,7 +846,7 @@ FVulkanRayTracingShaderTable::FVulkanRayTracingShaderTable(FRHICommandListBase&
 				Alloc.Region.size = Alloc.HandleCount * (uint32)GVulkanRayTracingMaxShaderGroupStride;
 
 				// Host buffer
-				Alloc.HostBuffer.SetNumUninitialized(Alloc.Region.size);
+				Alloc.HostBuffer.SetNumZeroed(Alloc.Region.size);
 			}
 		};
 
