PR #13438: Support for Qualcomm's FDM Offset Extension: VK_EXT_fragment_density_map_offset
Adjustments from original PR:
- Renamed GRHI globals to be consistent with existing VRS globals
- Submit offsets every pass to avoid underdetermined behavior when submitting offsets to one pass but no offsets to a subsequent pass
- Add r.Vulkan.AllowFDMOffset CVar to control extension setup and make extension setup flow consistent with other Vulkan extensions

#Synced-CL 50657708

diff --git a/Engine/Config/BaseDeviceProfiles.ini b/Engine/Config/BaseDeviceProfiles.ini
--- a/Engine/Config/BaseDeviceProfiles.ini
+++ b/Engine/Config/BaseDeviceProfiles.ini
@@ -1403,6 +1403,7 @@ BaseProfileName=Android_Mid
 +CVars=r.ForceStartInVR=1
 +CVars=xr.VRS.DynamicFoveation=1
 +CVars=r.Vulkan.VRSFormat=3
++CVars=r.Vulkan.AllowFDMOffset=1
 +CVars=sg.ShadowQuality=2
 
 [Meta_Quest_3 DeviceProfile]
@@ -1517,17 +1518,24 @@ BaseProfileName=Android_Mid
 ; cv is Steam Frame
 [cv DeviceProfile]
 DeviceType=LinuxArm64
-BaseProfileName=LinuxArm64
+BaseProfileName=LinuxArm64Client
 +CVars=r.AllowOcclusionQueries=False
 +CVars=r.DefaultBackBufferPixelFormat=0
 +CVars=fx.Niagara.QualityLevel.Min=0
 +CVars=fx.Niagara.QualityLevel.Max=3
 +CVars=r.Mobile.UseHWsRGBEncoding=1
++CVars=r.RDG.transientallocator=0
++CVars=vr.MobileMultiView=1
++CVars=r.DistanceFields=0
++CVars=r.Vulkan.VRSFormat=3
++CVars=r.Vulkan.AllowFDMOffset=1
 
 [Valve_Frame_Lepton DeviceProfile]
 DeviceType=Android
 BaseProfileName=Meta_Quest_3
 +CVars=xr.DisableOpenXROnAndroidWithoutOculus=0
++CVars=r.DistanceFields=0
++CVars=r.shadow.CSM.MaxCascades=0
 
 [Windows_Preview_ES31 DeviceProfile]
 DeviceType=Windows
diff --git a/Engine/Plugins/Runtime/OpenXR/Source/OpenXRHMD/Private/OpenXRHMDModule.cpp b/Engine/Plugins/Runtime/OpenXR/Source/OpenXRHMD/Private/OpenXRHMDModule.cpp
--- a/Engine/Plugins/Runtime/OpenXR/Source/OpenXRHMD/Private/OpenXRHMDModule.cpp
+++ b/Engine/Plugins/Runtime/OpenXR/Source/OpenXRHMD/Private/OpenXRHMDModule.cpp
@@ -570,6 +570,7 @@ bool FOpenXRHMDModule::GetOptionalExtensions(TArray<const ANSICHAR*>& OutExtensi
 	OutExtensions.Add(XR_KHR_VULKAN_ENABLE_EXTENSION_NAME);
 	OutExtensions.Add(XR_KHR_VULKAN_SWAPCHAIN_FORMAT_LIST_EXTENSION_NAME);
 	OutExtensions.Add(XR_FB_FOVEATION_VULKAN_EXTENSION_NAME);
+	OutExtensions.Add(XR_META_VULKAN_SWAPCHAIN_CREATE_INFO_EXTENSION_NAME);
 #endif
 	OutExtensions.Add(XR_KHR_COMPOSITION_LAYER_DEPTH_EXTENSION_NAME);
 	OutExtensions.Add(XR_FB_COMPOSITION_LAYER_DEPTH_TEST_EXTENSION_NAME);
diff --git a/Engine/Plugins/Runtime/OpenXR/Source/OpenXRHMD/Private/OpenXRHMD_Swapchain.cpp b/Engine/Plugins/Runtime/OpenXR/Source/OpenXRHMD/Private/OpenXRHMD_Swapchain.cpp
--- a/Engine/Plugins/Runtime/OpenXR/Source/OpenXRHMD/Private/OpenXRHMD_Swapchain.cpp
+++ b/Engine/Plugins/Runtime/OpenXR/Source/OpenXRHMD/Private/OpenXRHMD_Swapchain.cpp
@@ -473,12 +473,22 @@ FXRSwapChainPtr CreateSwapchain_Vulkan(XrSession InSession, uint8 Format, uint8&
 	void* Next = !(CreateFlags & TexCreate_SRGB) ? &FormatListInfo : nullptr;
 
 	XrSwapchainCreateInfoFoveationFB FoveationCreateInfo { XR_TYPE_SWAPCHAIN_CREATE_INFO_FOVEATION_FB };
+	XrVulkanSwapchainCreateInfoMETA VulkanMetaCreateInfo { XR_TYPE_VULKAN_SWAPCHAIN_CREATE_INFO_META };
 	if (EnumHasAllFlags(AuxiliaryCreateFlags, TexCreate_Foveation))
 	{
 		FoveationCreateInfo.next = Next;
 		FoveationCreateInfo.flags = XR_SWAPCHAIN_CREATE_FOVEATION_FRAGMENT_DENSITY_MAP_BIT_FB;
 
 		Next = &FoveationCreateInfo;
+
+		if (GRHISupportsVariableRateShadingImageOffsets)
+		{
+			VulkanMetaCreateInfo.next = &FoveationCreateInfo;
+			VulkanMetaCreateInfo.additionalUsageFlags = 0;
+			VulkanMetaCreateInfo.additionalCreateFlags = VK_IMAGE_CREATE_FRAGMENT_DENSITY_MAP_OFFSET_BIT_EXT;
+
+			Next = &VulkanMetaCreateInfo;
+		}
 	}
 
 	XrSwapchain Swapchain = FOpenXRSwapchain::CreateSwapchain(InSession, ViewFormatList[0], SizeX, SizeY, ArraySize, NumMips, NumSamples, CreateFlags, Next);
diff --git a/Engine/Source/Runtime/RHI/Public/RHIGlobals.h b/Engine/Source/Runtime/RHI/Public/RHIGlobals.h
--- a/Engine/Source/Runtime/RHI/Public/RHIGlobals.h
+++ b/Engine/Source/Runtime/RHI/Public/RHIGlobals.h
@@ -600,6 +600,15 @@ struct FRHIGlobals
 		/** Whether Variable Rate Shading deferred shading rate texture update is supported. */
 		bool SupportsLateUpdate = false;
 
+		/** Whether or not the RHI can support offsets for image-based Variable Rate Shading. Currently requires VK_EXT_fragment_density_map_offset to be supported and enabled. */
+		bool SupportsImageOffsets = false;
+		
+		/** Current offsets of the shading rate image, in pixels. NOTE: This MUST be quantized by the amount specified in ImageOffsetGranularity below! */
+		FIntPoint ImageOffsets[2];
+		
+		/** The granularity of the offset, in pixels, required by the RHI. This is queried at runtime, and must be used to enforce granularity for the ImageOffsets above. */
+		FIntPoint ImageOffsetGranularity;
+
 	} VariableRateShading;
 
 	/** Format used for the backbuffer when outputting to a HDR display. */
@@ -914,6 +923,9 @@ extern RHI_API FRHIGlobals GRHIGlobals;
 #define GRHIVariableRateShadingImageDataType                   GRHIGlobals.VariableRateShading.ImageDataType
 #define GRHIVariableRateShadingImageFormat                     GRHIGlobals.VariableRateShading.ImageFormat
 #define GRHISupportsLateVariableRateShadingUpdate              GRHIGlobals.VariableRateShading.SupportsLateUpdate
+#define GRHISupportsVariableRateShadingImageOffsets			   GRHIGlobals.VariableRateShading.SupportsImageOffsets
+#define GRHIVariableRateShadingImageOffsets					   GRHIGlobals.VariableRateShading.ImageOffsets
+#define GRHIVariableRateShadingImageOffsetGranularity		   GRHIGlobals.VariableRateShading.ImageOffsetGranularity
 #define GRHIHDRDisplayOutputFormat                             GRHIGlobals.HDRDisplayOutputFormat
 #define GRHIPresentCounter                                     GRHIGlobals.PresentCounter
 #define GRHISupportsArrayIndexFromAnyShader                    GRHIGlobals.SupportsArrayIndexFromAnyShader
diff --git a/Engine/Source/Runtime/VulkanRHI/Private/VulkanCommandBuffer.cpp b/Engine/Source/Runtime/VulkanRHI/Private/VulkanCommandBuffer.cpp
--- a/Engine/Source/Runtime/VulkanRHI/Private/VulkanCommandBuffer.cpp
+++ b/Engine/Source/Runtime/VulkanRHI/Private/VulkanCommandBuffer.cpp
@@ -91,7 +91,38 @@ void FVulkanCommandBuffer::FreeMemory()
 void FVulkanCommandBuffer::EndRenderPass()
 {
 	checkf(IsInsideRenderPass(), TEXT("Can't EndRP as we're NOT inside one! CmdBuffer 0x%p State=%d"), CommandBufferHandle, (int32)State);
-	VulkanRHI::vkCmdEndRenderPass(CommandBufferHandle);
+
+	if (Device.GetOptionalExtensions().HasKHRRenderPass2)
+	{
+		VkSubpassEndInfo SubpassEndInfo;
+		ZeroVulkanStruct(SubpassEndInfo, VK_STRUCTURE_TYPE_SUBPASS_END_INFO);
+
+		VkRenderPassFragmentDensityMapOffsetEndInfoEXT FDMOffsetInfo;
+		if (Device.GetOptionalExtensions().HasEXTFragmentDensityMapOffset && GRHISupportsVariableRateShadingImageOffsets)
+		{
+			VkOffset2D Offsets[2];
+
+			Offsets[0].x = GRHIVariableRateShadingImageOffsets[0].X;
+			Offsets[0].y = GRHIVariableRateShadingImageOffsets[0].Y;
+			Offsets[1].x = GRHIVariableRateShadingImageOffsets[1].X;
+			Offsets[1].y = GRHIVariableRateShadingImageOffsets[1].Y;
+
+			ZeroVulkanStruct(FDMOffsetInfo, VK_STRUCTURE_TYPE_RENDER_PASS_FRAGMENT_DENSITY_MAP_OFFSET_END_INFO_EXT);
+			FDMOffsetInfo.pFragmentDensityOffsets = Offsets;
+			FDMOffsetInfo.fragmentDensityOffsetCount = 2;
+			FDMOffsetInfo.pFragmentDensityOffsets = Offsets;
+			FDMOffsetInfo.pNext = nullptr;
+ 		
+			SubpassEndInfo.pNext = &FDMOffsetInfo;
+		}
+		
+		VulkanRHI::vkCmdEndRenderPass2KHR(CommandBufferHandle, &SubpassEndInfo);
+	}
+	else
+	{
+		VulkanRHI::vkCmdEndRenderPass(CommandBufferHandle);
+	}
+
 	State = EState::IsInsideBegin;
 	CurrentMultiViewCount = 0;
 }
@@ -142,7 +173,35 @@ void FVulkanCommandBuffer::BeginDynamicRendering(const VkRenderingInfo& Renderin
 void FVulkanCommandBuffer::EndDynamicRendering()
 {
 	check(GetCommandBufferType() != EVulkanCommandBufferType::Secondary);
-	VulkanRHI::vkCmdEndRenderingKHR(CommandBufferHandle);
+
+	// vkCmdEndRendering2KHR is provided by VK_EXT_fragment_density_map_offset to attach the offset info struct
+	if (Device.GetOptionalExtensions().HasEXTFragmentDensityMapOffset && GRHISupportsVariableRateShadingImageOffsets)
+	{
+		VkRenderingEndInfoKHR SubpassEndInfo;
+		ZeroVulkanStruct(SubpassEndInfo, VK_STRUCTURE_TYPE_RENDERING_END_INFO_KHR);
+
+		VkRenderPassFragmentDensityMapOffsetEndInfoEXT FDMOffsetInfo;
+		VkOffset2D Offsets[2];
+
+		Offsets[0].x = GRHIVariableRateShadingImageOffsets[0].X;
+		Offsets[0].y = GRHIVariableRateShadingImageOffsets[0].Y;
+		Offsets[1].x = GRHIVariableRateShadingImageOffsets[1].X;
+		Offsets[1].y = GRHIVariableRateShadingImageOffsets[1].Y;
+
+		ZeroVulkanStruct(FDMOffsetInfo, VK_STRUCTURE_TYPE_RENDER_PASS_FRAGMENT_DENSITY_MAP_OFFSET_END_INFO_EXT);
+		FDMOffsetInfo.pFragmentDensityOffsets = Offsets;
+		FDMOffsetInfo.fragmentDensityOffsetCount = 2;
+		FDMOffsetInfo.pFragmentDensityOffsets = Offsets;
+		FDMOffsetInfo.pNext = nullptr;
+ 		
+		SubpassEndInfo.pNext = &FDMOffsetInfo;
+		VulkanRHI::vkCmdEndRendering2KHR(CommandBufferHandle, &SubpassEndInfo);
+	}
+	else
+	{
+		VulkanRHI::vkCmdEndRenderingKHR(CommandBufferHandle);
+	}
+
 	State = EState::IsInsideBegin;
 }
 
diff --git a/Engine/Source/Runtime/VulkanRHI/Private/VulkanCommandWrappers.h b/Engine/Source/Runtime/VulkanRHI/Private/VulkanCommandWrappers.h
--- a/Engine/Source/Runtime/VulkanRHI/Private/VulkanCommandWrappers.h
+++ b/Engine/Source/Runtime/VulkanRHI/Private/VulkanCommandWrappers.h
@@ -192,8 +192,10 @@ struct FWrapLayer
 	static void CmdBeginRenderPass2KHR(VkResult Result, VkCommandBuffer CommandBuffer, const VkRenderPassBeginInfo* RenderPassBegin, const VkSubpassBeginInfo* pSubpassBeginInfo) VULKAN_LAYER_BODY
 	static void CmdNextSubpass(VkResult Result, VkCommandBuffer CommandBuffer, VkSubpassContents Contents) VULKAN_LAYER_BODY
 	static void CmdEndRenderPass(VkResult Result, VkCommandBuffer CommandBuffer) VULKAN_LAYER_BODY
+	static void CmdEndRenderPass2KHR(VkResult Result, VkCommandBuffer CommandBuffer, const VkSubpassEndInfo* SubpassEndInfo) VULKAN_LAYER_BODY
 	static void CmdBeginRenderingKHR(VkResult Result, VkCommandBuffer CommandBuffer, const VkRenderingInfo* RenderingInfo) VULKAN_LAYER_BODY
 	static void CmdEndRenderingKHR(VkResult Result, VkCommandBuffer CommandBuffer) VULKAN_LAYER_BODY
+	static void CmdEndRendering2KHR(VkResult Result, VkCommandBuffer CommandBuffer, const VkRenderingEndInfoKHR* RenderingEndInfo) VULKAN_LAYER_BODY
 	static void CmdExecuteCommands(VkResult Result, VkCommandBuffer CommandBuffer, uint32 CommandBufferCount, const VkCommandBuffer* pCommandBuffers) VULKAN_LAYER_BODY
 
 	static void QueuePresent(VkResult Result, VkQueue Queue, const VkPresentInfoKHR* PresentInfo) VULKAN_LAYER_BODY
@@ -1451,6 +1453,13 @@ namespace VulkanRHI
 		VULKANAPINAMESPACE::vkCmdEndRenderPass(CommandBuffer);
 		FWrapLayer::CmdEndRenderPass(VK_SUCCESS, CommandBuffer);
 	}
+	// The VK_EXT_fragment_density_map_offset extension requires EndRenderPass2KHR so that we can pack extra info into the struct chain to specify the offsets
+ 	static FORCEINLINE_DEBUGGABLE void vkCmdEndRenderPass2KHR(VkCommandBuffer CommandBuffer, const VkSubpassEndInfo* SubpassEndInfo)
+ 	{
+		FWrapLayer::CmdEndRenderPass2KHR(VK_RESULT_MAX_ENUM, CommandBuffer, SubpassEndInfo);
+ 		VULKANAPINAMESPACE::vkCmdEndRenderPass2KHR(CommandBuffer, SubpassEndInfo);
+ 		FWrapLayer::CmdEndRenderPass2KHR(VK_SUCCESS, CommandBuffer, SubpassEndInfo);
+ 	}
 
 	static FORCEINLINE_DEBUGGABLE void vkCmdBeginRenderingKHR(VkCommandBuffer CommandBuffer, const VkRenderingInfo* RenderingInfo)
 	{
@@ -1466,6 +1475,13 @@ namespace VulkanRHI
 		FWrapLayer::CmdEndRenderingKHR(VK_SUCCESS, CommandBuffer);
 	}
 
+	static FORCEINLINE_DEBUGGABLE void vkCmdEndRendering2KHR(VkCommandBuffer CommandBuffer, const VkRenderingEndInfoKHR* RenderingEndInfo)
+	{
+		FWrapLayer::CmdEndRendering2KHR(VK_RESULT_MAX_ENUM, CommandBuffer, RenderingEndInfo);
+		VULKANAPINAMESPACE::vkCmdEndRendering2KHR(CommandBuffer, RenderingEndInfo);
+		FWrapLayer::CmdEndRendering2KHR(VK_SUCCESS, CommandBuffer, RenderingEndInfo);
+	}
+
 	static FORCEINLINE_DEBUGGABLE void  vkCmdExecuteCommands(VkCommandBuffer CommandBuffer, uint32 CommandBufferCount, const VkCommandBuffer* pCommandBuffers)
 	{
 		FWrapLayer::CmdExecuteCommands(VK_RESULT_MAX_ENUM, CommandBuffer, CommandBufferCount, pCommandBuffers);
diff --git a/Engine/Source/Runtime/VulkanRHI/Private/VulkanDebug.cpp b/Engine/Source/Runtime/VulkanRHI/Private/VulkanDebug.cpp
--- a/Engine/Source/Runtime/VulkanRHI/Private/VulkanDebug.cpp
+++ b/Engine/Source/Runtime/VulkanRHI/Private/VulkanDebug.cpp
@@ -2542,6 +2542,16 @@ void FWrapLayer::CmdEndRenderPass(VkResult Result, VkCommandBuffer CommandBuffer
 	}
 }
 
+void FWrapLayer::CmdEndRenderPass2KHR(VkResult Result, VkCommandBuffer CommandBuffer)
+{
+	if (Result == VK_RESULT_MAX_ENUM)
+	{
+#if VULKAN_ENABLE_DUMP_LAYER
+		CmdPrintfBegin(CommandBuffer, TEXT("----- vkCmdEndRenderPass2KHR()"));
+#endif
+	}
+}
+
 void FWrapLayer::CmdNextSubpass(VkResult Result, VkCommandBuffer CommandBuffer, VkSubpassContents Contents)
 {
 	if (Result == VK_RESULT_MAX_ENUM)
@@ -2572,6 +2582,16 @@ void FWrapLayer::CmdEndRenderingKHR(VkResult Result, VkCommandBuffer CommandBuff
 	}
 }
 
+void FWrapLayer::CmdEndRendering2KHR(VkResult Result, VkCommandBuffer CommandBuffer)
+{
+	if (Result == VK_RESULT_MAX_ENUM)
+	{
+#if VULKAN_ENABLE_DUMP_LAYER
+		CmdPrintfBegin(CommandBuffer, TEXT("----- vkCmdEndRendering2KHR()"));
+#endif
+	}
+}
+
 void FWrapLayer::CmdExecuteCommands(VkResult Result, VkCommandBuffer CommandBuffer, uint32 CommandBufferCount, const VkCommandBuffer* pCommandBuffers)
 {
 	if (Result == VK_RESULT_MAX_ENUM)
diff --git a/Engine/Source/Runtime/VulkanRHI/Private/VulkanDevice.h b/Engine/Source/Runtime/VulkanRHI/Private/VulkanDevice.h
--- a/Engine/Source/Runtime/VulkanRHI/Private/VulkanDevice.h
+++ b/Engine/Source/Runtime/VulkanRHI/Private/VulkanDevice.h
@@ -48,6 +48,7 @@ struct FOptionalVulkanDeviceExtensions
 			uint64 HasEXTASTCDecodeMode : 1;
 			uint64 HasEXTFragmentDensityMap : 1;
 			uint64 HasEXTFragmentDensityMap2 : 1;
+			uint64 HasEXTFragmentDensityMapOffset : 1;
 			uint64 HasKHRFragmentShadingRate : 1;
 			uint64 HasKHRFragmentShaderBarycentric : 1;
 			uint64 HasEXTFullscreenExclusive : 1;
diff --git a/Engine/Source/Runtime/VulkanRHI/Private/VulkanDynamicRHI.h b/Engine/Source/Runtime/VulkanRHI/Private/VulkanDynamicRHI.h
--- a/Engine/Source/Runtime/VulkanRHI/Private/VulkanDynamicRHI.h
+++ b/Engine/Source/Runtime/VulkanRHI/Private/VulkanDynamicRHI.h
@@ -316,6 +316,8 @@ public:
 	// Runs code on SubmissionThread with access to VkQueue.  Useful for plugins. 
 	virtual void RHIRunOnQueue(EVulkanRHIRunOnQueueType QueueType, TFunction<void(VkQueue)>&& CodeToRun, bool bWaitForSubmission) final override;
 
+	virtual void RHISetFDMOffsets(const FVector2f GazePoint[2], FIntPoint FramebufferSize);
+
 	void DumpCrashMarkers();
 
 	static void SavePipelineCache();
diff --git a/Engine/Source/Runtime/VulkanRHI/Private/VulkanExtensions.cpp b/Engine/Source/Runtime/VulkanRHI/Private/VulkanExtensions.cpp
--- a/Engine/Source/Runtime/VulkanRHI/Private/VulkanExtensions.cpp
+++ b/Engine/Source/Runtime/VulkanRHI/Private/VulkanExtensions.cpp
@@ -146,6 +146,15 @@ TAutoConsoleVariable<int32> GVulkanVariableRateShadingFormatCVar(
 	ECVF_ReadOnly
 );
 
+TAutoConsoleVariable<int32> GVulkanAllowFDMOffsetCVar(
+	TEXT("r.Vulkan.AllowFDMOffset"),
+	0,
+	TEXT("Enables the use of VK_EXT_fragment_density_map_offset\n")
+	TEXT("0: Do not enable support for fragment_density_map_offset (default)\n")
+	TEXT("1: Enable fragment_density_map_offset if available"),
+	ECVF_ReadOnly
+);
+
 TAutoConsoleVariable<int32> GVulkanAllowImageRobustness(
 	TEXT("r.Vulkan.AllowImageRobustness"),
 	1,
@@ -942,6 +951,50 @@ public:
 	}
 };
 
+class FVulkanEXTFragmentDensityMapOffsetExtension : public FVulkanDeviceExtension
+{
+public:
+
+	FVulkanEXTFragmentDensityMapOffsetExtension(FVulkanDevice* InDevice)
+		: FVulkanDeviceExtension(InDevice, VK_EXT_FRAGMENT_DENSITY_MAP_OFFSET_EXTENSION_NAME, VULKAN_EXTENSION_ENABLED)
+	{
+		bEnabledInCode = bEnabledInCode && 
+			HardwareVariableRateShadingSupportedByPlatform(GMaxRHIShaderPlatform) && GVulkanAllowFDMOffsetCVar.GetValueOnAnyThread();
+	}
+
+	virtual void PrePhysicalDeviceFeatures(VkPhysicalDeviceFeatures2KHR& PhysicalDeviceFeatures2) override final
+	{
+		ZeroVulkanStruct(FragmentDensityMapOffsetFeatures, VK_STRUCTURE_TYPE_PHYSICAL_DEVICE_FRAGMENT_DENSITY_MAP_OFFSET_FEATURES_EXT);
+		AddToPNext(PhysicalDeviceFeatures2, FragmentDensityMapOffsetFeatures);
+	}
+
+	virtual void PostPhysicalDeviceFeatures(FOptionalVulkanDeviceExtensions& ExtensionFlags) override final
+	{
+		bRequirementsPassed = (FragmentDensityMapOffsetFeatures.fragmentDensityMapOffset == VK_TRUE) && (ExtensionFlags.HasEXTFragmentDensityMap || ExtensionFlags.HasEXTFragmentDensityMap2);
+		ExtensionFlags.HasEXTFragmentDensityMapOffset = bRequirementsPassed;
+	}
+
+	virtual void PrePhysicalDeviceProperties(VkPhysicalDeviceProperties2KHR& PhysicalDeviceProperties2) override final
+	{
+		ZeroVulkanStruct(FragmentDensityMapOffsetProperties, VK_STRUCTURE_TYPE_PHYSICAL_DEVICE_FRAGMENT_DENSITY_MAP_OFFSET_PROPERTIES_EXT);
+		AddToPNext(PhysicalDeviceProperties2, FragmentDensityMapOffsetProperties);
+	}
+
+	virtual void PostPhysicalDeviceProperties(FOptionalVulkanDeviceExtensions& ExtensionFlags) override final
+	{
+		if (bRequirementsPassed)
+		{
+			GRHISupportsVariableRateShadingImageOffsets = true;
+
+			GRHIVariableRateShadingImageOffsetGranularity.X = FragmentDensityMapOffsetProperties.fragmentDensityOffsetGranularity.width;
+			GRHIVariableRateShadingImageOffsetGranularity.Y = FragmentDensityMapOffsetProperties.fragmentDensityOffsetGranularity.height;
+		}
+	}
+
+private:
+	VkPhysicalDeviceFragmentDensityMapOffsetFeaturesEXT FragmentDensityMapOffsetFeatures;
+	VkPhysicalDeviceFragmentDensityMapOffsetPropertiesEXT FragmentDensityMapOffsetProperties;
+};
 
 
 // ***** VK_KHR_fragment_shader_barycentric
@@ -2597,6 +2650,7 @@ FVulkanDeviceExtensionArray FVulkanDeviceExtension::GetUESupportedDeviceExtensio
 	ADD_CUSTOM_EXTENSION(FVulkanKHRFragmentShadingRateExtension); // must be kept BEFORE DensityMap!
 	ADD_CUSTOM_EXTENSION(FVulkanEXTFragmentDensityMapExtension);  // must be kept AFTER ShadingRate!
 	ADD_CUSTOM_EXTENSION(FVulkanEXTFragmentDensityMap2Extension);
+	ADD_CUSTOM_EXTENSION(FVulkanEXTFragmentDensityMapOffsetExtension)
 	ADD_CUSTOM_EXTENSION(FVulkanKHRMultiviewExtension);
 	ADD_CUSTOM_EXTENSION(FVulkanKHRGetMemoryRequirements2Extension);
 	ADD_CUSTOM_EXTENSION(FVulkanEXTDescriptorIndexingExtension);
diff --git a/Engine/Source/Runtime/VulkanRHI/Private/VulkanLoader.h b/Engine/Source/Runtime/VulkanRHI/Private/VulkanLoader.h
--- a/Engine/Source/Runtime/VulkanRHI/Private/VulkanLoader.h
+++ b/Engine/Source/Runtime/VulkanRHI/Private/VulkanLoader.h
@@ -149,6 +149,7 @@
 	EnumMacro(PFN_vkCmdBeginRenderPass, vkCmdBeginRenderPass) \
 	EnumMacro(PFN_vkCmdNextSubpass, vkCmdNextSubpass) \
 	EnumMacro(PFN_vkCmdEndRenderPass, vkCmdEndRenderPass) \
+	EnumMacro(PFN_vkCmdEndRenderPass2KHR, vkCmdEndRenderPass2KHR) \
 	EnumMacro(PFN_vkCmdExecuteCommands, vkCmdExecuteCommands) \
 	EnumMacro(PFN_vkCreateSwapchainKHR, vkCreateSwapchainKHR) \
 	EnumMacro(PFN_vkDestroySwapchainKHR, vkDestroySwapchainKHR) \
@@ -206,6 +207,7 @@
 	EnumMacro(PFN_vkCmdBeginRenderPass2KHR, vkCmdBeginRenderPass2KHR) \
 	EnumMacro(PFN_vkCmdBeginRenderingKHR, vkCmdBeginRenderingKHR) \
 	EnumMacro(PFN_vkCmdEndRenderingKHR, vkCmdEndRenderingKHR) \
+	EnumMacro(PFN_vkCmdEndRendering2KHR, vkCmdEndRendering2KHR) \
 	EnumMacro(PFN_vkGetDeviceImageMemoryRequirementsKHR, vkGetDeviceImageMemoryRequirementsKHR) \
 	EnumMacro(PFN_vkGetDeviceBufferMemoryRequirementsKHR, vkGetDeviceBufferMemoryRequirementsKHR) \
 	EnumMacro(PFN_vkResetQueryPoolEXT, vkResetQueryPoolEXT) \
diff --git a/Engine/Source/Runtime/VulkanRHI/Private/VulkanRHI.cpp b/Engine/Source/Runtime/VulkanRHI/Private/VulkanRHI.cpp
--- a/Engine/Source/Runtime/VulkanRHI/Private/VulkanRHI.cpp
+++ b/Engine/Source/Runtime/VulkanRHI/Private/VulkanRHI.cpp
@@ -2133,5 +2133,29 @@ void FVulkanDynamicRHI::RHIUpdateResourceCollection(FRHICommandListBase& RHICmdL
 }
 #endif
 
+void FVulkanDynamicRHI::RHISetFDMOffsets(const FVector2f GazePoint[2], FIntPoint FramebufferSize)
+{
+	if (GRHISupportsVariableRateShadingImageOffsets)
+	{
+		// Eye Offsets are in -1.0..1.0 space
+		const float LeftXOffset = FMath::Clamp(GazePoint[0].X, -1.0, 1.0);
+		const float LeftYOffset = FMath::Clamp(GazePoint[0].Y, -1.0, 1.0);
+		const float RightXOffset = FMath::Clamp(GazePoint[1].X, -1.0, 1.0);
+		const float RightYOffset = FMath::Clamp(GazePoint[1].Y, -1.0, 1.0);
+
+		// Sanity check, these two should never be 0.  If so, for some reason the extension property initialization has failed
+		check(GRHIVariableRateShadingImageOffsetGranularity.X > 0 && GRHIVariableRateShadingImageOffsetGranularity.Y > 0);
+
+		// Offsets need to be quantized to the pixel size defined in the extension's properties struct
+		const FVector2D QuantizedGranularity((float)FramebufferSize.X / (float)GRHIVariableRateShadingImageOffsetGranularity.X,
+			(float)FramebufferSize.Y / (float)GRHIVariableRateShadingImageOffsetGranularity.Y);
+
+		// Quantize the offsets to the minimum offset allowed by the FDM Offset granularity
+		GRHIVariableRateShadingImageOffsets[0].X = FMath::RoundToInt(LeftXOffset * 0.5 * QuantizedGranularity.X) * GRHIVariableRateShadingImageOffsetGranularity.X;
+		GRHIVariableRateShadingImageOffsets[0].Y = FMath::RoundToInt(LeftYOffset * 0.5 * QuantizedGranularity.Y) * GRHIVariableRateShadingImageOffsetGranularity.Y;
+		GRHIVariableRateShadingImageOffsets[1].X = FMath::RoundToInt(RightXOffset * 0.5 * QuantizedGranularity.X) * GRHIVariableRateShadingImageOffsetGranularity.X;
+		GRHIVariableRateShadingImageOffsets[1].Y = FMath::RoundToInt(RightYOffset * 0.5 * QuantizedGranularity.Y) * GRHIVariableRateShadingImageOffsetGranularity.Y;
+	}
+}
 
 #undef LOCTEXT_NAMESPACE
diff --git a/Engine/Source/Runtime/VulkanRHI/Private/VulkanTexture.cpp b/Engine/Source/Runtime/VulkanRHI/Private/VulkanTexture.cpp
--- a/Engine/Source/Runtime/VulkanRHI/Private/VulkanTexture.cpp
+++ b/Engine/Source/Runtime/VulkanRHI/Private/VulkanTexture.cpp
@@ -371,6 +371,11 @@ void FVulkanTexture::GenerateImageCreateInfo(
 		ImageCreateInfo.flags |= VK_IMAGE_CREATE_2D_ARRAY_COMPATIBLE_BIT;
 	}
 
+	if (InDevice.GetOptionalExtensions().HasEXTFragmentDensityMapOffset && GRHISupportsVariableRateShadingImageOffsets)
+	{
+		ImageCreateInfo.flags |= VK_IMAGE_CREATE_FRAGMENT_DENSITY_MAP_OFFSET_BIT_EXT;
+	}
+
 	ImageCreateInfo.tiling = bForceLinearTexture ? VK_IMAGE_TILING_LINEAR : GVulkanViewTypeTilingMode[ResourceType];
 	if (EnumHasAnyFlags(UEFlags, TexCreate_RenderTargetable | TexCreate_DepthStencilTargetable | TexCreate_ResolveTargetable | TexCreate_DepthStencilResolveTarget))
 	{
diff --git a/Engine/Source/Runtime/VulkanRHI/Public/IVulkanDynamicRHI.h b/Engine/Source/Runtime/VulkanRHI/Public/IVulkanDynamicRHI.h
--- a/Engine/Source/Runtime/VulkanRHI/Public/IVulkanDynamicRHI.h
+++ b/Engine/Source/Runtime/VulkanRHI/Public/IVulkanDynamicRHI.h
@@ -76,6 +76,8 @@ struct IVulkanDynamicRHI : public FDynamicRHI
 	virtual VkFormat         RHIGetSwapChainVkFormat(EPixelFormat InFormat) const = 0;
 	virtual bool             RHISupportsEXTFragmentDensityMap2() const = 0;
 
+	virtual void 			 RHISetFDMOffsets(const FVector2f GazePoint[2], FIntPoint FramebufferSize) = 0;
+
 	virtual TArray<VkExtensionProperties> RHIGetAllInstanceExtensions() const = 0;
 	virtual TArray<VkExtensionProperties> RHIGetAllDeviceExtensions(VkPhysicalDevice InPhysicalDevice) const = 0;
 	virtual TArray<FAnsiString> RHIGetLoadedDeviceExtensions() const = 0;
