VulkanRHI: Add support for specialization constants:
- Add specialization constant values to pipeline hashes to have unique states per specialization combo.
- Use FComputePipelineStateInitializer internally to create compute pipelines (to include spec const values with shader).
- Specify specialization constant values when create graphics or compute pipelines
- Move AddToPNext to VulkanCommon so it can be used everywhere.
- Add a SpecializationCount to Vulkan shader code header
- Adjust internal descriptor validation in FVulkanDescriptorSetWriter to take excluded resources for specialization into account
- Add a pass in SPIRV compilation to evaluate the effects of specialization constant combinations to resource bindings.  Generate a list of hashes for which to exclude a given resource.

#rnx

#Synced-CL 50785456

diff --git a/Engine/Shaders/Public/Platform/Vulkan/VulkanCommon.ush b/Engine/Shaders/Public/Platform/Vulkan/VulkanCommon.ush
--- a/Engine/Shaders/Public/Platform/Vulkan/VulkanCommon.ush
+++ b/Engine/Shaders/Public/Platform/Vulkan/VulkanCommon.ush
@@ -7,7 +7,7 @@
 #pragma once
 
 // Update this GUID to invalidate shader recompilation for only Vulkan shaders
-#pragma message("UESHADERMETADATA_VERSION A6FA7DFC-55D8-4F8C-A56C-C3B3AD4645C3")
+#pragma message("UESHADERMETADATA_VERSION 62D5B8B1-DECD-4C2D-98A4-6E5FE07DEA2E")
 
 #define COMPILER_SUPPORTS_ULONG_TYPES 1
 #define COMPILER_SUPPORTS_UINT64_IMAGE_ATOMICS (1 && COMPILER_SUPPORTS_ULONG_TYPES)
@@ -169,4 +169,10 @@ struct FVulkanPushConstants
 	uint StaticUniformBufferBindings[8];  // keep in sync with FRHIShaderBindingLayout::MaxUniformBufferEntries
 };
 [[vk::push_constant]] ConstantBuffer<FVulkanPushConstants> VulkanPushConstants;
-#endif
\ No newline at end of file
+#endif
+
+// Specialization constant support
+#if PLATFORM_SUPPORTS_SPECIALIZATION_CONSTANTS && ENABLE_SPECIALIZATION_CONSTANTS
+#define UE_DECL(Name) [[vk::constant_id(Name##_INDEX)]] const int Name = 0
+#define UE_EVAL(Name) Name
+#endif
diff --git a/Engine/Source/Developer/VulkanShaderFormat/Private/SpirVShaderCompiler.inl b/Engine/Source/Developer/VulkanShaderFormat/Private/SpirVShaderCompiler.inl
--- a/Engine/Source/Developer/VulkanShaderFormat/Private/SpirVShaderCompiler.inl
+++ b/Engine/Source/Developer/VulkanShaderFormat/Private/SpirVShaderCompiler.inl
@@ -5,6 +5,14 @@
 #include "Serialization/MemoryWriter.h"
 #include "Misc/FileHelper.h"
 #include "Misc/Paths.h"
+#include "ShaderCompilerCommonInternal.h"
+
+
+//THIRD_PARTY_INCLUDES_START
+//#include <spirv/unified1/GLSL.std.450.h>
+//#include <../Include/spirv-tools/libspirv.h>
+//THIRD_PARTY_INCLUDES_END
+
 
 // A collection of states and data that is locked in at the top level call and doesn't change throughout the compilation process
 class FSpirvShaderCompilerInternalState
@@ -669,6 +677,7 @@ static uint32 CalculateSpirvInstructionCount(FSpirv& Spirv)
 	return ApproxInstructionCount;
 }
 
+
 static bool BuildShaderOutputFromSpirv(
 	CrossCompiler::FShaderConductorContext&	CompilerContext,
 	const FSpirvShaderCompilerInternalState& InternalState,
@@ -963,6 +972,9 @@ static bool BuildShaderOutputFromSpirv(
 	// Overwrite updated SPIRV code
 	SerializedOutput.Spirv.Data = TArray<uint32>(Reflection.GetCode(), Reflection.GetCodeSize() / 4);
 
+	// After we've copied the SPIRV module's copy of the SPIRV, use it to evaluate specialization constants
+	SerializedOutput.Header.SpecializationCount = UE::ShaderCompilerCommon::ProcessSpirvSpecializationConstants(InternalState.Input, CompilerContext, Reflection, Bindings, Output);
+
 	// We have to strip out most debug instructions (except OpName) for Vulkan mobile
 	if (InternalState.ShouldStripReflect())
 	{
diff --git a/Engine/Source/Developer/VulkanShaderFormat/Private/VulkanShaderCompiler.cpp b/Engine/Source/Developer/VulkanShaderFormat/Private/VulkanShaderCompiler.cpp
--- a/Engine/Source/Developer/VulkanShaderFormat/Private/VulkanShaderCompiler.cpp
+++ b/Engine/Source/Developer/VulkanShaderFormat/Private/VulkanShaderCompiler.cpp
@@ -371,6 +371,7 @@ void CompileVulkanShader(const FShaderCompilerInput& Input, const FShaderPreproc
 	}
 
 	Output.SerializeShaderCodeValidation();
+	Output.SerializeSpecializationConstantData();
 
 	ShaderParameterParser.ValidateShaderParameterTypes(Input, InternalState.IsMobileES31(), Output);
 	
diff --git a/Engine/Source/Runtime/VulkanRHI/Private/VulkanDescriptorSets.cpp b/Engine/Source/Runtime/VulkanRHI/Private/VulkanDescriptorSets.cpp
--- a/Engine/Source/Runtime/VulkanRHI/Private/VulkanDescriptorSets.cpp
+++ b/Engine/Source/Runtime/VulkanRHI/Private/VulkanDescriptorSets.cpp
@@ -450,24 +450,29 @@ void FVulkanDescriptorSetWriter::CheckAllWritten()
 	if (Writes == 0)
 		return;
 
+	auto GetWrittenMask = [WrittenMask=WrittenMask, ExcludedMask=ExcludedMask](int32 Index)
+	{
+			return WrittenMask[Index] | ExcludedMask[Index];
+	};
+
 	bool bFail = false;
 	if (Writes <= 32) //early out for the most common case.
 	{
-		bFail = WrittenMask[0] != ((1llu << Writes) - 1);
+		bFail = GetWrittenMask(0) != ((1llu << Writes) - 1);
 	}
 	else
 	{
 		const int32 Last = int32(WrittenMask.Num() - 1);
 		for (int32 i = 0; !bFail && i < Last; ++i)
 		{
-			uint64 Mask = WrittenMask[i];
+			const uint64 Mask = GetWrittenMask(i);
 			bFail = bFail || Mask != 0xffffffff;
 		}
 
 		const uint32 TailCount = Writes - (Last * 32);
 		check(TailCount != 0);
 		const uint32 TailMask = (1llu << TailCount) - 1;
-		bFail = bFail || TailMask != WrittenMask[Last];
+		bFail = bFail || TailMask != GetWrittenMask(Last);
 	}
 
 	if (bFail)
@@ -522,5 +527,19 @@ void FVulkanDescriptorSetWriter::InitWrittenMasks(uint32 NumDescriptorWrites)
 	WrittenMask.SetNumZeroed(Size);
 	BaseWrittenMask.Empty(Size);
 	BaseWrittenMask.SetNumZeroed(Size);
+	ExcludedMask.Empty(Size);
+	ExcludedMask.SetNumZeroed(Size);
+#endif
+}
+
+void FVulkanDescriptorSetWriter::SetExcludedResources(const TArrayView<uint32> ExcludedResourceCount)
+{
+#if VULKAN_VALIDATE_DESCRIPTORS_WRITTEN	
+	for (uint32 DescriptorIndex : ExcludedResourceCount)
+	{
+		const uint32 Index = DescriptorIndex / 32;
+		const uint32 Mask = DescriptorIndex % 32;
+		ExcludedMask[Index] |= (1 << Mask);
+	}
 #endif
 }
diff --git a/Engine/Source/Runtime/VulkanRHI/Private/VulkanDescriptorSets.h b/Engine/Source/Runtime/VulkanRHI/Private/VulkanDescriptorSets.h
--- a/Engine/Source/Runtime/VulkanRHI/Private/VulkanDescriptorSets.h
+++ b/Engine/Source/Runtime/VulkanRHI/Private/VulkanDescriptorSets.h
@@ -12,7 +12,7 @@
 #if (UE_BUILD_SHIPPING || UE_BUILD_TEST)
 #define VULKAN_VALIDATE_DESCRIPTORS_WRITTEN 0
 #else
-#define VULKAN_VALIDATE_DESCRIPTORS_WRITTEN 1
+#define VULKAN_VALIDATE_DESCRIPTORS_WRITTEN RHI_INCLUDE_SHADER_DEBUG_DATA
 #endif
 
 
@@ -826,12 +826,14 @@ protected:
 #if VULKAN_VALIDATE_DESCRIPTORS_WRITTEN
 	TArray<uint32, TInlineAllocator<2> > WrittenMask;
 	TArray<uint32, TInlineAllocator<2> > BaseWrittenMask;
+	TArray<uint32, TInlineAllocator<2> > ExcludedMask;
 #endif
 	void CheckAllWritten();
 	void Reset();
 	void SetWritten(uint32 DescriptorIndex);
 	void SetWrittenBase(uint32 DescriptorIndex);
 	void InitWrittenMasks(uint32 NumDescriptorWrites);
+	void SetExcludedResources(const TArrayView<uint32> ExcludedResourceCount);
 };
 
 
diff --git a/Engine/Source/Runtime/VulkanRHI/Private/VulkanExtensions.cpp b/Engine/Source/Runtime/VulkanRHI/Private/VulkanExtensions.cpp
--- a/Engine/Source/Runtime/VulkanRHI/Private/VulkanExtensions.cpp
+++ b/Engine/Source/Runtime/VulkanRHI/Private/VulkanExtensions.cpp
@@ -182,14 +182,6 @@ TArray<const ANSICHAR*> FVulkanDeviceExtension::ExternalExtensions;
 TArray<const ANSICHAR*> FVulkanInstanceExtension::ExternalExtensions;
 
 
-template <typename ExistingChainType, typename NewStructType>
-static void AddToPNext(ExistingChainType& Existing, NewStructType& Added)
-{
-	Added.pNext = (void*)Existing.pNext;
-	Existing.pNext = (void*)&Added;
-}
-
-
 struct FOptionalVulkanDeviceExtensionProperties& FVulkanDeviceExtension::GetDeviceExtensionProperties()
 {
 	const FOptionalVulkanDeviceExtensionProperties& ExtensionProperties = Device->GetOptionalExtensionProperties();
diff --git a/Engine/Source/Runtime/VulkanRHI/Private/VulkanPipeline.cpp b/Engine/Source/Runtime/VulkanRHI/Private/VulkanPipeline.cpp
--- a/Engine/Source/Runtime/VulkanRHI/Private/VulkanPipeline.cpp
+++ b/Engine/Source/Runtime/VulkanRHI/Private/VulkanPipeline.cpp
@@ -245,10 +245,11 @@ FVulkanPipeline::~FVulkanPipeline()
 	/* we do NOT own Layout !*/
 }
 
-FVulkanComputePipeline::FVulkanComputePipeline(FVulkanDevice* InDevice, FVulkanComputeShader* InShader)
+FVulkanComputePipeline::FVulkanComputePipeline(FVulkanDevice* InDevice, const FComputePipelineStateInitializer& Initializer)
 	: FVulkanPipeline(InDevice)
-	, FRHIComputePipelineState(InShader)
-	, bUsesBindless(InShader->UsesBindless())
+	, FRHIComputePipelineState(Initializer.ComputeShader)
+	, bUsesBindless(ResourceCast(Initializer.ComputeShader)->UsesBindless())
+	, SpecializationConstants(Initializer.SpecializationConstants)
 {
 	INC_DWORD_STAT(STAT_VulkanNumComputePSOs);
 }
@@ -321,6 +322,16 @@ FVulkanPipelineStateCacheManager::FVulkanPipelineStateCacheManager(FVulkanDevice
 		CreateVertexInputLibraries();
 		CreateEmptyFragmentStageLibrary();
 	}
+
+	constexpr uint32 SpecializationConstantSize = sizeof(uint32);
+	for (uint32 MapEntryIndex = 0; MapEntryIndex < MaxNumSpecializationConstants; ++MapEntryIndex)
+	{
+		SpecializationMapEntries[MapEntryIndex] = {
+			.constantID = MapEntryIndex,
+			.offset = MapEntryIndex * SpecializationConstantSize,
+			.size = SpecializationConstantSize
+		};
+	}
 }
 
 
@@ -1187,6 +1198,8 @@ FArchive& operator << (FArchive& Ar, FGfxPipelineDesc& Entry)
 		Ar << Entry.ShaderHashes.Stages[Index];
 	}
 #endif
+	Ar << Entry.SpecializationKey;
+
 	Ar << Entry.RenderTargets;
 
 	uint8 ShadingRate = static_cast<uint8>(Entry.ShadingRate);
@@ -1724,6 +1737,7 @@ bool FVulkanPipelineStateCacheManager::CreateGfxPipelineFromEntry(FVulkanGraphic
 	// main_00000000_00000000
 	ANSICHAR EntryPoints[ShaderStage::NumGraphicsStages][24];
 	VkPipelineShaderStageRequiredSubgroupSizeCreateInfo RequiredSubgroupSizeCreateInfo[ShaderStage::NumGraphicsStages];
+	VkSpecializationInfo SpecializationInfo[ShaderStage::NumGraphicsStages];
 	for (int32 ShaderStage = 0; ShaderStage < ShaderStage::NumGraphicsStages; ++ShaderStage)
 	{
 		if (!ShaderModules[ShaderStage].IsValid() || (Shaders[ShaderStage] == nullptr))
@@ -1750,13 +1764,30 @@ bool FVulkanPipelineStateCacheManager::CreateGfxPipelineFromEntry(FVulkanGraphic
 				const bool bSupportedSize = ((ShaderHeader.WaveSize >= SubgroupSizeControlProperties.minSubgroupSize) && (ShaderHeader.WaveSize <= SubgroupSizeControlProperties.maxSubgroupSize));
 				if (bSupportedStage && bSupportedSize)
 				{
-				ZeroVulkanStruct(RequiredSubgroupSizeCreateInfo[PipelineInfo.stageCount], VK_STRUCTURE_TYPE_PIPELINE_SHADER_STAGE_REQUIRED_SUBGROUP_SIZE_CREATE_INFO);
-				RequiredSubgroupSizeCreateInfo[PipelineInfo.stageCount].requiredSubgroupSize = ShaderHeader.WaveSize;
+					RequiredSubgroupSizeCreateInfo[PipelineInfo.stageCount] = { 
+						.sType = VK_STRUCTURE_TYPE_PIPELINE_SHADER_STAGE_REQUIRED_SUBGROUP_SIZE_CREATE_INFO,
+						.pNext = nullptr,
+						.requiredSubgroupSize = ShaderHeader.WaveSize
+					};
 				ShaderStages[PipelineInfo.stageCount].pNext = &RequiredSubgroupSizeCreateInfo[PipelineInfo.stageCount];
 			}
 		}
 		}
 
+		const uint32 NumSpecConsts = PSO->SpecializationConstants[ShaderStage].Num();
+		checkf(NumSpecConsts < MaxNumSpecializationConstants, TEXT("Exceeded maximum number of specialization constants!  Raise FVulkanPipelineStateCacheManager::MaxNumSpecializationConstants."));
+		if (NumSpecConsts)
+		{
+			SpecializationInfo[ShaderStage] = {
+				.mapEntryCount = NumSpecConsts,
+				.pMapEntries = SpecializationMapEntries.GetData(),
+				.dataSize = PSO->SpecializationConstants[ShaderStage].NumBytes(),
+				.pData = PSO->SpecializationConstants[ShaderStage].GetData()
+			};
+
+			ShaderStages[PipelineInfo.stageCount].pSpecializationInfo = &SpecializationInfo[ShaderStage];
+		}
+
 		PipelineInfo.stageCount++;
 	}
 
@@ -2474,6 +2505,28 @@ void FVulkanPipelineStateCacheManager::CreateGfxEntry(const FGraphicsPipelineSta
 #endif
 	check(NumShaders > 0);
 
+	// Include a hash of specialization constants used with PSO
+	{
+		const FBoundShaderStateInput& BSI = PSOInitializer.BoundShaderState;
+		TMemoryHasher<FXxHash64Builder, FXxHash64> HashAr;
+
+#define VK_HASH_SPECIALIZATION_ARRAY(Array)	if (!Array.IsEmpty()) HashAr.Serialize((void*)Array.GetData(), Array.NumBytes())
+
+		VK_HASH_SPECIALIZATION_ARRAY(BSI.VertexShaderSpecializationConstants);
+		VK_HASH_SPECIALIZATION_ARRAY(BSI.PixelShaderSpecializationConstants);
+#if VULKAN_SUPPORTS_GEOMETRY_SHADERS
+		VK_HASH_SPECIALIZATION_ARRAY(BSI.GeometryShaderSpecializationConstants);
+#endif
+#if PLATFORM_SUPPORTS_MESH_SHADERS
+		VK_HASH_SPECIALIZATION_ARRAY(BSI.MeshShaderSpecializationConstants);
+		VK_HASH_SPECIALIZATION_ARRAY(BSI.AmplificationShaderSpecializationConstants);
+#endif
+
+#undef VK_HASH_SPECIALIZATION_ARRAY
+
+		OutGfxEntry->SpecializationKey = HashAr.Finalize().Hash;
+	}
+
 	FVulkanRenderTargetLayout RTLayout(PSOInitializer);
 	OutGfxEntry->RenderTargets.ReadFrom(RTLayout);
 
@@ -2500,14 +2553,19 @@ FVulkanGraphicsPipelineState::FVulkanGraphicsPipelineState(FVulkanDevice* InDevi
 
 	FMemory::Memset(VulkanShaders, 0, sizeof(VulkanShaders));
 	VulkanShaders[ShaderStage::Vertex] = static_cast<FVulkanVertexShader*>(PSOInitializer_.BoundShaderState.VertexShaderRHI);
+	SpecializationConstants[ShaderStage::Vertex] = PSOInitializer_.BoundShaderState.VertexShaderSpecializationConstants;
 #if PLATFORM_SUPPORTS_MESH_SHADERS
 	VulkanShaders[ShaderStage::Mesh] = static_cast<FVulkanMeshShader*>(PSOInitializer_.BoundShaderState.GetMeshShader());
+	SpecializationConstants[ShaderStage::Mesh] = PSOInitializer_.BoundShaderState.MeshShaderSpecializationConstants;
 	VulkanShaders[ShaderStage::Task] = static_cast<FVulkanTaskShader*>(PSOInitializer_.BoundShaderState.GetAmplificationShader());
+	SpecializationConstants[ShaderStage::Task] = PSOInitializer_.BoundShaderState.AmplificationShaderSpecializationConstants;
 #endif
 #if PLATFORM_SUPPORTS_GEOMETRY_SHADERS
 	VulkanShaders[ShaderStage::Geometry] = static_cast<FVulkanGeometryShader*>(PSOInitializer_.BoundShaderState.GetGeometryShader());
+	SpecializationConstants[ShaderStage::Geometry] = PSOInitializer_.BoundShaderState.GeometryShaderSpecializationConstants;
 #endif
 	VulkanShaders[ShaderStage::Pixel] = static_cast<FVulkanPixelShader*>(PSOInitializer_.BoundShaderState.PixelShaderRHI);
+	SpecializationConstants[ShaderStage::Pixel] = PSOInitializer_.BoundShaderState.PixelShaderSpecializationConstants;
 
 	uint32 ActiveShaderCount = 0;
 	uint32 BindlessShaderCount = 0;
@@ -2749,10 +2807,9 @@ FGraphicsPipelineStateRHIRef FVulkanDynamicRHI::RHICreateGraphicsPipelineState(c
 	return Device->PipelineStateCache->RHICreateGraphicsPipelineState(PSOInitializer);
 }
 
-FVulkanComputePipeline* FVulkanPipelineStateCacheManager::RHICreateComputePipelineState(FRHIComputeShader* ComputeShaderRHI)
+FVulkanComputePipeline* FVulkanPipelineStateCacheManager::RHICreateComputePipelineState(const FComputePipelineStateInitializer& Initializer)
 {
-	FVulkanComputeShader* ComputeShader = ResourceCast(ComputeShaderRHI);
-	return Device->GetPipelineStateCache()->GetOrCreateComputePipeline(ComputeShader);
+	return Device->GetPipelineStateCache()->GetOrCreateComputePipeline(Initializer);
 }
 
 FComputePipelineStateRHIRef FVulkanDynamicRHI::RHICreateComputePipelineState(const FComputePipelineStateInitializer& Initializer)
@@ -2763,13 +2820,16 @@ FComputePipelineStateRHIRef FVulkanDynamicRHI::RHICreateComputePipelineState(con
 	QUICK_SCOPE_CYCLE_COUNTER(STAT_Vulkan_RHICreateComputePipelineState);
 	LLM_SCOPE_VULKAN(ELLMTagVulkan::VulkanShaders);
 
-	return Device->PipelineStateCache->RHICreateComputePipelineState(Initializer.ComputeShader);
+	return Device->PipelineStateCache->RHICreateComputePipelineState(Initializer);
 }
 
-FVulkanComputePipeline* FVulkanPipelineStateCacheManager::GetOrCreateComputePipeline(FVulkanComputeShader* ComputeShader)
+FVulkanComputePipeline* FVulkanPipelineStateCacheManager::GetOrCreateComputePipeline(const FComputePipelineStateInitializer& Initializer)
 {
+	FVulkanComputeShader* ComputeShader = ResourceCast(Initializer.ComputeShader);
 	check(ComputeShader);
-	const uint64 Key = ComputeShader->GetShaderKey();
+	const FSpecializationContainerType& SpecializationConstants = Initializer.SpecializationConstants;
+	const uint64 Key = CityHash64WithSeed((const char*)SpecializationConstants.GetData(), SpecializationConstants.NumBytes(), ComputeShader->GetShaderKey());
+
 	{
 		FRWScopeLock ScopeLock(ComputePipelineLock, SLT_ReadOnly);
 		FVulkanComputePipeline** ComputePipelinePtr = ComputePipelineEntries.Find(Key);
@@ -2782,7 +2842,7 @@ FVulkanComputePipeline* FVulkanPipelineStateCacheManager::GetOrCreateComputePipe
 	// create pipeline of entry + store entry
 	double BeginTime = FPlatformTime::Seconds();
 
-	FVulkanComputePipeline* ComputePipeline = CreateComputePipelineFromShader(ComputeShader);
+	FVulkanComputePipeline* ComputePipeline = CreateComputePipeline(Initializer);
 
 	double EndTime = FPlatformTime::Seconds();
 	double Delta = EndTime - BeginTime;
@@ -2802,9 +2862,11 @@ FVulkanComputePipeline* FVulkanPipelineStateCacheManager::GetOrCreateComputePipe
 	return ComputePipeline;
 }
 
-FVulkanComputePipeline* FVulkanPipelineStateCacheManager::CreateComputePipelineFromShader(FVulkanComputeShader* Shader)
+FVulkanComputePipeline* FVulkanPipelineStateCacheManager::CreateComputePipeline(const FComputePipelineStateInitializer& Initializer)
 {
-	FVulkanComputePipeline* Pipeline = new FVulkanComputePipeline(Device, Shader);
+	FVulkanComputeShader* Shader = ResourceCast(Initializer.ComputeShader);
+	const FSpecializationContainerType& SpecializationConstants = Initializer.SpecializationConstants;
+	FVulkanComputePipeline* Pipeline = new FVulkanComputePipeline(Device, Initializer);
 
 	FVulkanDescriptorSetsLayoutInfo DescriptorSetLayoutInfo;
 	const FVulkanShaderHeader& CSHeader = Shader->GetCodeHeader();
@@ -2822,22 +2884,47 @@ FVulkanComputePipeline* FVulkanPipelineStateCacheManager::CreateComputePipelineF
 
 	TRefCountPtr<FVulkanShaderModule> ShaderModule = Shader->GetOrCreateHandle(Layout, Layout->GetDescriptorSetLayoutHash());
 
-	VkComputePipelineCreateInfo PipelineInfo;
-	ZeroVulkanStruct(PipelineInfo, VK_STRUCTURE_TYPE_COMPUTE_PIPELINE_CREATE_INFO);
-	PipelineInfo.stage.sType = VK_STRUCTURE_TYPE_PIPELINE_SHADER_STAGE_CREATE_INFO;
-	PipelineInfo.stage.stage = VK_SHADER_STAGE_COMPUTE_BIT;
-	PipelineInfo.stage.module = ShaderModule->GetVkShaderModule();
-	// main_00000000_00000000
+	// Entrypoint: main_00000000_00000000
 	ANSICHAR EntryPoint[24];
 	Shader->GetEntryPoint(EntryPoint, 24);
-	PipelineInfo.stage.pName = EntryPoint;
-	PipelineInfo.layout = Layout->GetPipelineLayout();
+
+	VkComputePipelineCreateInfo PipelineInfo = {
+		.sType = VK_STRUCTURE_TYPE_COMPUTE_PIPELINE_CREATE_INFO,
+		.pNext = nullptr,
+		.flags = 0,
+		.stage = {
+			.sType = VK_STRUCTURE_TYPE_PIPELINE_SHADER_STAGE_CREATE_INFO,
+			.pNext = nullptr,
+			.flags = 0,
+			.stage = VK_SHADER_STAGE_COMPUTE_BIT,
+			.module = ShaderModule->GetVkShaderModule(),
+			.pName = EntryPoint,
+			.pSpecializationInfo = nullptr
+		},
+		.layout = Layout->GetPipelineLayout(),
+		.basePipelineHandle = VK_NULL_HANDLE,
+		.basePipelineIndex = 0
+	};
 
 	if (Shader->UsesBindless())
 	{
 		PipelineInfo.flags |= VK_PIPELINE_CREATE_DESCRIPTOR_BUFFER_BIT_EXT;
 	}
 
+	VkSpecializationInfo SpecializationInfo;
+	if (SpecializationConstants.Num())
+	{
+		checkf(SpecializationConstants.Num() < MaxNumSpecializationConstants, TEXT("Exceeded maximum number of specialization constants!  Raise FVulkanPipelineStateCacheManager::MaxNumSpecializationConstants."));
+		SpecializationInfo = {
+			.mapEntryCount = (uint32)SpecializationConstants.Num(),
+			.pMapEntries = SpecializationMapEntries.GetData(),
+			.dataSize = SpecializationConstants.NumBytes(),
+			.pData = SpecializationConstants.GetData()
+		};
+
+		PipelineInfo.stage.pSpecializationInfo = &SpecializationInfo;
+	}
+
 	VkPipelineShaderStageRequiredSubgroupSizeCreateInfo RequiredSubgroupSizeCreateInfo;
 	if ((CSHeader.WaveSize > 0) && Device->GetOptionalExtensions().HasEXTSubgroupSizeControl)
 	{
@@ -2848,9 +2935,13 @@ FVulkanComputePipeline* FVulkanPipelineStateCacheManager::CreateComputePipelineF
 			// Check if requested size is supported
 			if ((CSHeader.WaveSize >= SubgroupSizeControlProperties.minSubgroupSize) && (CSHeader.WaveSize <= SubgroupSizeControlProperties.maxSubgroupSize))
 			{
-				ZeroVulkanStruct(RequiredSubgroupSizeCreateInfo, VK_STRUCTURE_TYPE_PIPELINE_SHADER_STAGE_REQUIRED_SUBGROUP_SIZE_CREATE_INFO);
-				RequiredSubgroupSizeCreateInfo.requiredSubgroupSize = CSHeader.WaveSize;
-				PipelineInfo.stage.pNext = &RequiredSubgroupSizeCreateInfo;
+				RequiredSubgroupSizeCreateInfo = {
+					.sType = VK_STRUCTURE_TYPE_PIPELINE_SHADER_STAGE_REQUIRED_SUBGROUP_SIZE_CREATE_INFO,
+					.pNext = nullptr,
+					.requiredSubgroupSize = CSHeader.WaveSize
+				};
+
+				AddToPNext(PipelineInfo.stage, RequiredSubgroupSizeCreateInfo);
 			}
 		}
 	}
@@ -2880,7 +2971,8 @@ void FVulkanPipelineStateCacheManager::NotifyDeletedComputePipeline(FVulkanCompu
 	FVulkanComputeShader* ComputeShader = ResourceCast(Pipeline->GetComputeShader());
 	if (ComputeShader)
 	{
-		const uint64 Key = ComputeShader->GetShaderKey();
+		const FSpecializationContainerType& SpecializationConstants = Pipeline->GetSpecializationConstants();
+		const uint64 Key = CityHash64WithSeed((const char*)SpecializationConstants.GetData(), SpecializationConstants.NumBytes(), ComputeShader->GetShaderKey());
 		FRWScopeLock ScopeLock(ComputePipelineLock, SLT_Write); 
 		ComputePipelineEntries.Remove(Key);
 	}
diff --git a/Engine/Source/Runtime/VulkanRHI/Private/VulkanPipeline.h b/Engine/Source/Runtime/VulkanRHI/Private/VulkanPipeline.h
--- a/Engine/Source/Runtime/VulkanRHI/Private/VulkanPipeline.h
+++ b/Engine/Source/Runtime/VulkanRHI/Private/VulkanPipeline.h
@@ -275,6 +275,7 @@ struct FGfxPipelineDesc
 #else
 	FVulkanShaderHashes ShaderHashes;
 #endif
+	uint64 SpecializationKey;
 
 	struct FRenderTargets
 	{
@@ -456,6 +457,11 @@ struct FGfxPipelineDesc
 		}
 #endif
 
+		if (SpecializationKey != In.SpecializationKey)
+		{
+			return false;
+		}
+
 		if (!(RenderTargets == In.RenderTargets))
 		{
 			return false;
@@ -501,7 +507,7 @@ public:
 
 	void RebuildCache();
 
-	FVulkanComputePipeline* GetOrCreateComputePipeline(FVulkanComputeShader* ComputeShader);
+	FVulkanComputePipeline* GetOrCreateComputePipeline(const FComputePipelineStateInitializer& Initializer);
 	void NotifyDeletedComputePipeline(FVulkanComputePipeline* Pipeline);
 
 	VkResult CreateRayTracingPipeline(const VkRayTracingPipelineCreateInfoKHR& CreateInfo, bool bIsPartial, VkPipeline& OutPipeline);
@@ -519,7 +525,7 @@ private:
 	void DestroyCache();
 
 	FGraphicsPipelineStateRHIRef RHICreateGraphicsPipelineState(const FGraphicsPipelineStateInitializer& Initializer);
-	FVulkanComputePipeline* RHICreateComputePipelineState(FRHIComputeShader* ComputeShaderRHI);
+	FVulkanComputePipeline* RHICreateComputePipelineState(const FComputePipelineStateInitializer& Initializer);
 	void NotifyDeletedGraphicsPSO(FRHIGraphicsPipelineState* PSO);
 	bool CreateGfxPipelineFromEntry(FVulkanGraphicsPipelineState* PSO, FVulkanShader* Shaders[ShaderStage::NumGraphicsStages], FGraphicsPipelineStateInitializer::EPSOPrecacheCompileType PSOCompileType);
 	bool CreateGfxPipelineFromLibrary(FVulkanGraphicsPipelineState& GraphicsPipelineState, FVulkanShader* Shaders[ShaderStage::NumGraphicsStages], FGraphicsPipelineStateInitializer::EPSOPrecacheCompileType PSOCompileType);
@@ -538,7 +544,7 @@ private:
 	VkResult CreateVKPipeline(FVulkanGraphicsPipelineState* PSO, FVulkanShader* Shaders[ShaderStage::NumGraphicsStages], const VkGraphicsPipelineCreateInfo& PipelineInfo, FGraphicsPipelineStateInitializer::EPSOPrecacheCompileType PSOCompileType);
 
 	FVulkanLayout* FindOrAddLayout(const FVulkanDescriptorSetsLayoutInfo& DescriptorSetLayoutInfo, bool bGfxLayout, bool bUsesBindless);
-	FVulkanComputePipeline* CreateComputePipelineFromShader(FVulkanComputeShader* Shader);
+	FVulkanComputePipeline* CreateComputePipeline(const FComputePipelineStateInitializer& Initializer);
 
 	/** LRU Related functions */
 	void TickLRU();
@@ -637,6 +643,9 @@ private:
 		void Save(FArchive& Ar);
 		bool Load(FArchive& Ar);
 	};
+
+	static constexpr uint32 MaxNumSpecializationConstants = 32;
+	TStaticArray<VkSpecializationMapEntry, MaxNumSpecializationConstants> SpecializationMapEntries;
 };
 
 struct VulkanResourceFrameCounter
@@ -704,7 +713,7 @@ protected:
 class FVulkanComputePipeline : public FVulkanPipeline, public FRHIComputePipelineState
 {
 public:
-	FVulkanComputePipeline(FVulkanDevice* InDevice, FVulkanComputeShader* InShader);
+	FVulkanComputePipeline(FVulkanDevice* InDevice, const FComputePipelineStateInitializer& Initializer);
 	virtual ~FVulkanComputePipeline();
 
 	const FVulkanShaderHeader& GetShaderCodeHeader() const
@@ -724,8 +733,14 @@ public:
 		return bUsesBindless;
 	}
 
+	const FSpecializationContainerType& GetSpecializationConstants() const
+	{
+		return SpecializationConstants;
+	}
+
 protected:
 	const bool bUsesBindless;
+	const FSpecializationContainerType SpecializationConstants;
 };
 
 class FVulkanGraphicsPipelineState : public FRHIGraphicsPipelineState
@@ -770,8 +785,8 @@ public:
 		case SF_Vertex: return static_cast<FVulkanVertexShader*>(VulkanShaders[ShaderStage::Vertex]);
 		case SF_Pixel: return static_cast<FVulkanPixelShader*>(VulkanShaders[ShaderStage::Pixel]);
 #if PLATFORM_SUPPORTS_MESH_SHADERS
-		case SF_Mesh: return static_cast<FVulkanGeometryShader*>(VulkanShaders[ShaderStage::Mesh]);
-		case SF_Amplification: return static_cast<FVulkanGeometryShader*>(VulkanShaders[ShaderStage::Task]);
+		case SF_Mesh: return static_cast<FVulkanMeshShader*>(VulkanShaders[ShaderStage::Mesh]);
+		case SF_Amplification: return static_cast<FVulkanTaskShader*>(VulkanShaders[ShaderStage::Task]);
 #endif
 #if VULKAN_SUPPORTS_GEOMETRY_SHADERS
 		case SF_Geometry: return static_cast<FVulkanGeometryShader*>(VulkanShaders[ShaderStage::Geometry]);
@@ -805,6 +820,7 @@ public:
 	FVulkanDevice* Device;
 	const FGfxPipelineDesc Desc;
 	FVulkanShader* VulkanShaders[ShaderStage::NumGraphicsStages];
+	FSpecializationContainerType SpecializationConstants[ShaderStage::NumGraphicsStages];
 	const FVulkanRenderPass* RenderPass;
 
 
diff --git a/Engine/Source/Runtime/VulkanRHI/Private/VulkanPipelineState.cpp b/Engine/Source/Runtime/VulkanRHI/Private/VulkanPipelineState.cpp
--- a/Engine/Source/Runtime/VulkanRHI/Private/VulkanPipelineState.cpp
+++ b/Engine/Source/Runtime/VulkanRHI/Private/VulkanPipelineState.cpp
@@ -132,6 +132,8 @@ void FVulkanCommonPipelineDescriptorState::CreateDescriptorWriteInfos()
 			CurrentAccelerationStructures,
 			DefaultSampler, DefaultImageView);
 
+		DSWriter[Set].SetExcludedResources(GetExcludedResources(Set));
+
 		TotalNumDynamicOffsets += NumDynamicOffsets;
 
 		if (CurrentHashableDescriptorInfo) // UseVulkanDescriptorCache()
@@ -283,6 +285,30 @@ void FVulkanComputePipelineDescriptorState::UpdateBindlessDescriptors(FVulkanCom
 	Device.GetBindlessDescriptorManager()->RegisterUniformBuffers(Context, VK_SHADER_STAGE_COMPUTE_BIT, UBWrites);
 }
 
+TArray<uint32> FVulkanComputePipelineDescriptorState::GetExcludedResources(uint32 SetIndex) const
+{
+	check(SetIndex == 0);
+
+	TArray<uint32> ExcludedResources;
+#if RHI_INCLUDE_SHADER_DEBUG_DATA
+	FRHIShader* RHIShader = ComputePipeline->GetComputeShader();
+	if (RHIShader && RHIShader->HasSpecializationConstants())
+	{
+		const uint32 SpecializationHash = GetTypeHash(ComputePipeline->GetSpecializationConstants());
+
+		for (auto& Pair : RHIShader->DebugSpecializationExclusions)
+		{
+			if (Pair.Value.Contains(SpecializationHash))
+			{
+				ExcludedResources.Add(Pair.Key);
+			}
+		}
+	}
+#endif
+	return ExcludedResources;
+}
+
+
 FVulkanGraphicsPipelineDescriptorState::FVulkanGraphicsPipelineDescriptorState(FVulkanDevice& InDevice, FVulkanGraphicsPipelineState* InGfxPipeline)
 	: FVulkanCommonPipelineDescriptorState(InDevice, ShaderStage::NumGraphicsStages, InGfxPipeline->UsesBindless())
 	, GfxPipeline(InGfxPipeline)
@@ -527,7 +553,28 @@ void FVulkanGraphicsPipelineDescriptorState::UpdateBindlessDescriptors(FVulkanCo
 	Device.GetBindlessDescriptorManager()->RegisterUniformBuffers(Context, StageFlags, UBCompactedWrites);
 }
 
+TArray<uint32> FVulkanGraphicsPipelineDescriptorState::GetExcludedResources(uint32 SetIndex) const
+{
+	check(SetIndex < ShaderStage::NumGraphicsStages);
+
+	TArray<uint32> ExcludedResources;
+#if RHI_INCLUDE_SHADER_DEBUG_DATA
+	FRHIShader* RHIShader = GfxPipeline->GetShader(GetFrequencyForGfxStage((ShaderStage::EStage)SetIndex));
+	if (RHIShader && RHIShader->HasSpecializationConstants())
+	{
+		const uint32 SpecializationHash = GetTypeHash(GfxPipeline->SpecializationConstants[SetIndex]);
 
+		for (auto& Pair : RHIShader->DebugSpecializationExclusions)
+		{
+			if (Pair.Value.Contains(SpecializationHash))
+			{
+				ExcludedResources.Add(Pair.Key);
+			}
+		}
+	}
+#endif
+	return ExcludedResources;
+}
 
 
 
diff --git a/Engine/Source/Runtime/VulkanRHI/Private/VulkanPipelineState.h b/Engine/Source/Runtime/VulkanRHI/Private/VulkanPipelineState.h
--- a/Engine/Source/Runtime/VulkanRHI/Private/VulkanPipelineState.h
+++ b/Engine/Source/Runtime/VulkanRHI/Private/VulkanPipelineState.h
@@ -121,6 +121,8 @@ public:
 		return Layouts[DescriptorSet].LayoutBindings[BindingIndex].descriptorType;
 	}
 
+	virtual TArray<uint32> GetExcludedResources(uint32 SetIndex) const = 0;
+
 protected:
 	void Reset()
 	{
@@ -241,6 +243,8 @@ public:
 		Bind(CmdBuffer, ComputePipeline->GetLayout().GetPipelineLayout(), VK_PIPELINE_BIND_POINT_COMPUTE);
 	}
 
+	TArray<uint32> GetExcludedResources(uint32 SetIndex) const override final;
+
 protected:
 	FVulkanPackedUniformBufferState PackedUniformBufferState;
 
@@ -305,6 +309,8 @@ public:
 		bIsResourcesDirty = true;
 	}
 
+	TArray<uint32> GetExcludedResources(uint32 SetIndex) const override final;
+
 protected:
 	TStaticArray<FVulkanPackedUniformBufferState, ShaderStage::NumGraphicsStages> PackedUniformBufferStates;
 
diff --git a/Engine/Source/Runtime/VulkanRHI/Private/VulkanShaders.cpp b/Engine/Source/Runtime/VulkanRHI/Private/VulkanShaders.cpp
--- a/Engine/Source/Runtime/VulkanRHI/Private/VulkanShaders.cpp
+++ b/Engine/Source/Runtime/VulkanRHI/Private/VulkanShaders.cpp
@@ -247,21 +247,26 @@ ShaderType* FVulkanShaderFactory::CreateShader(TArrayView<const uint8> Code, FVu
 				FShaderCodeReader ShaderCode(Code);
 				ReadShaderOptionalData(ShaderCode, RetShader);
 
+				if (FDataDrivenShaderPlatformInfo::GetSupportsSpecializationConstants(GMaxRHIShaderPlatform))
+				{
+					RetShader->SetSpecializationConstants(RetShader->GetCodeHeader().SpecializationCount > 0);
+				}
+
 				if (Device->SupportsShaderObjects())
 				{
 					Device->GetShaderObjectManager()->CreateShaderObject(*RetShader);
 				}
 				else if (Device->SupportsGraphicPipelineLibraries())
 				{
-					if (std::is_same_v<ShaderType, FVulkanVertexShader>)
+					if constexpr (std::is_same_v<ShaderType, FVulkanVertexShader>)
 					{
 						RetShader->PipelineLibrary = CreatePreRasterShaderLibrary(Device, RetShader, VK_SHADER_STAGE_VERTEX_BIT);
 					}
-					if (std::is_same_v<ShaderType, FVulkanMeshShader>)
+					if constexpr (std::is_same_v<ShaderType, FVulkanMeshShader>)
 					{
 						RetShader->PipelineLibrary = CreatePreRasterShaderLibrary(Device, RetShader, VK_SHADER_STAGE_MESH_BIT_EXT);
 					}
-					if (std::is_same_v<ShaderType, FVulkanPixelShader>)
+					if constexpr (std::is_same_v<ShaderType, FVulkanPixelShader>)
 					{
 						RetShader->PipelineLibrary = CreatePixelShaderLibrary(Device, RetShader);
 					}
diff --git a/Engine/Source/Runtime/VulkanRHI/Public/VulkanCommon.h b/Engine/Source/Runtime/VulkanRHI/Public/VulkanCommon.h
--- a/Engine/Source/Runtime/VulkanRHI/Public/VulkanCommon.h
+++ b/Engine/Source/Runtime/VulkanRHI/Public/VulkanCommon.h
@@ -167,3 +167,10 @@ static inline void ZeroVulkanStruct(T& Struct, int32 VkStructureType)
 	(int32&)Struct.sType = VkStructureType;
 	FMemory::Memzero(((uint8*)&Struct) + sizeof(VkStructureType), sizeof(T) - sizeof(VkStructureType));
 }
+
+template <typename ExistingChainType, typename NewStructType>
+static inline void AddToPNext(ExistingChainType& Existing, NewStructType& Added)
+{
+	Added.pNext = (void*)Existing.pNext;
+	Existing.pNext = (void*)&Added;
+}
diff --git a/Engine/Source/Runtime/VulkanRHI/Public/VulkanShaderResources.h b/Engine/Source/Runtime/VulkanRHI/Public/VulkanShaderResources.h
--- a/Engine/Source/Runtime/VulkanRHI/Public/VulkanShaderResources.h
+++ b/Engine/Source/Runtime/VulkanRHI/Public/VulkanShaderResources.h
@@ -102,6 +102,9 @@ struct FVulkanShaderHeader
 	uint32						RayTracingPayloadType = 0;
 	uint32						RayTracingPayloadSize = 0;
 
+	// Specialization constant info
+	uint32						SpecializationCount = 0;
+
 	FSHAHash					SourceHash;
 	uint32						SpirvCRC = 0;
 	uint8						WaveSize = 0;
@@ -179,6 +182,7 @@ inline FArchive& operator<<(FArchive& Ar, FVulkanShaderHeader& Header)
 	Ar << Header.InOutMask;
 	Ar << Header.RayTracingPayloadType;
 	Ar << Header.RayTracingPayloadSize;
+	Ar << Header.SpecializationCount;
 	Ar << Header.SourceHash;
 	Ar << Header.SpirvCRC;
 	Ar << Header.WaveSize;
