VulkanRHI: Remove more unused ImmutableSampler parameters

#rnx

#Synced-CL 50785456

diff --git a/Engine/Source/Runtime/VulkanRHI/Private/VulkanDescriptorSets.cpp b/Engine/Source/Runtime/VulkanRHI/Private/VulkanDescriptorSets.cpp
--- a/Engine/Source/Runtime/VulkanRHI/Private/VulkanDescriptorSets.cpp
+++ b/Engine/Source/Runtime/VulkanRHI/Private/VulkanDescriptorSets.cpp
@@ -17,7 +17,7 @@ void FVulkanDescriptorSetsLayoutInfo::ProcessBindingsForStage(VkShaderStageFlagB
 }
 
 template<bool bIsCompute>
-void FVulkanDescriptorSetsLayoutInfo::FinalizeBindings(const FVulkanDevice& Device, const FUniformBufferGatherInfo& UBGatherInfo, const TArrayView<FRHISamplerState*>& ImmutableSamplers, bool bUsesBindless)
+void FVulkanDescriptorSetsLayoutInfo::FinalizeBindings(const FVulkanDevice& Device, const FUniformBufferGatherInfo& UBGatherInfo, bool bUsesBindless)
 {
 	// We'll be reusing this struct
 	VkDescriptorSetLayoutBinding Binding;
@@ -28,7 +28,6 @@ void FVulkanDescriptorSetsLayoutInfo::FinalizeBindings(const FVulkanDevice& Devi
 	const bool bConvertPackedUBsToDynamic = !bUsesBindless && (bConvertAllUBsToDynamic || (GDynamicGlobalUBs.GetValueOnAnyThread() == 1));
 	const uint32 MaxDescriptorSetUniformBuffersDynamic = Device.GetLimits().maxDescriptorSetUniformBuffersDynamic;
 
-	int32 CurrentImmutableSampler = 0;
 	for (int32 Stage = 0; Stage < (bIsCompute ? ShaderStage::NumComputeStages : ShaderStage::NumGraphicsStages); ++Stage)
 	{
 		checkSlow(StageInfos[Stage].IsEmpty());
@@ -93,11 +92,11 @@ void FVulkanDescriptorSetsLayoutInfo::FinalizeBindings(const FVulkanDevice& Devi
 	}
 
 	CompileTypesUsageID();
-	GenerateHash(ImmutableSamplers, bIsCompute ? VK_PIPELINE_BIND_POINT_COMPUTE : VK_PIPELINE_BIND_POINT_GRAPHICS);
+	GenerateHash(bIsCompute ? VK_PIPELINE_BIND_POINT_COMPUTE : VK_PIPELINE_BIND_POINT_GRAPHICS);
 }
 
-template void FVulkanDescriptorSetsLayoutInfo::FinalizeBindings<true>(const FVulkanDevice& Device, const FUniformBufferGatherInfo& UBGatherInfo, const TArrayView<FRHISamplerState*>& ImmutableSamplers, bool bUsesBindless);
-template void FVulkanDescriptorSetsLayoutInfo::FinalizeBindings<false>(const FVulkanDevice& Device, const FUniformBufferGatherInfo& UBGatherInfo, const TArrayView<FRHISamplerState*>& ImmutableSamplers, bool bUsesBindless);
+template void FVulkanDescriptorSetsLayoutInfo::FinalizeBindings<true>(const FVulkanDevice& Device, const FUniformBufferGatherInfo& UBGatherInfo, bool bUsesBindless);
+template void FVulkanDescriptorSetsLayoutInfo::FinalizeBindings<false>(const FVulkanDevice& Device, const FUniformBufferGatherInfo& UBGatherInfo, bool bUsesBindless);
 
 
 // Increments a value and asserts on overflow.
@@ -160,7 +159,7 @@ void FVulkanDescriptorSetsLayoutInfo::AddDescriptor(int32 DescriptorSetIndex, co
 	}
 }
 
-void FVulkanDescriptorSetsLayoutInfo::GenerateHash(const TArrayView<FRHISamplerState*>& InImmutableSamplers, VkPipelineBindPoint InBindPoint)
+void FVulkanDescriptorSetsLayoutInfo::GenerateHash(VkPipelineBindPoint InBindPoint)
 {
 	const int32 LayoutCount = SetLayouts.Num();
 	Hash = FCrc::MemCrc32(&TypesUsageID, sizeof(uint32), LayoutCount);
diff --git a/Engine/Source/Runtime/VulkanRHI/Private/VulkanDescriptorSets.h b/Engine/Source/Runtime/VulkanRHI/Private/VulkanDescriptorSets.h
--- a/Engine/Source/Runtime/VulkanRHI/Private/VulkanDescriptorSets.h
+++ b/Engine/Source/Runtime/VulkanRHI/Private/VulkanDescriptorSets.h
@@ -105,9 +105,9 @@ public:
 	void ProcessBindingsForStage(VkShaderStageFlagBits StageFlags, ShaderStage::EStage DescSetStage, const FVulkanShaderHeader& CodeHeader, FUniformBufferGatherInfo& OutUBGatherInfo) const;
 
 	template<bool bIsCompute>
-	void FinalizeBindings(const FVulkanDevice& Device, const FUniformBufferGatherInfo& UBGatherInfo, const TArrayView<FRHISamplerState*>& ImmutableSamplers, bool bUsesBindless);
+	void FinalizeBindings(const FVulkanDevice& Device, const FUniformBufferGatherInfo& UBGatherInfo, bool bUsesBindless);
 
-	void GenerateHash(const TArrayView<FRHISamplerState*>& ImmutableSamplers, VkPipelineBindPoint InBindPoint);
+	void GenerateHash(VkPipelineBindPoint InBindPoint);
 
 	friend uint32 GetTypeHash(const FVulkanDescriptorSetsLayoutInfo& In)
 	{
diff --git a/Engine/Source/Runtime/VulkanRHI/Private/VulkanPipeline.cpp b/Engine/Source/Runtime/VulkanRHI/Private/VulkanPipeline.cpp
--- a/Engine/Source/Runtime/VulkanRHI/Private/VulkanPipeline.cpp
+++ b/Engine/Source/Runtime/VulkanRHI/Private/VulkanPipeline.cpp
@@ -2383,8 +2383,7 @@ void FVulkanPipelineStateCacheManager::CreateGfxEntry(const FGraphicsPipelineSta
 		checkf((NumBindlessShaders == 0) || (NumBindlessShaders == NumActiveShaders), TEXT("All shaders must be bindless or non-bindless."));
 
 		// Second pass
-		//todo: strip out ImmutableSamplers from 3rd parameter.
-		DescriptorSetLayoutInfo.FinalizeBindings<false>(*Device, UBGatherInfo, TArrayView<FRHISamplerState*>(), (NumBindlessShaders != 0));
+		DescriptorSetLayoutInfo.FinalizeBindings<false>(*Device, UBGatherInfo, (NumBindlessShaders != 0));
 	}
 
 	checkf(PSOInitializer.SubpassIndex == 0 || PSOInitializer.SubpassHint != ESubpassHint::None, TEXT("Subpass hint must be set if subpass index > 0."));
@@ -2872,7 +2871,7 @@ FVulkanComputePipeline* FVulkanPipelineStateCacheManager::CreateComputePipeline(
 	const FVulkanShaderHeader& CSHeader = Shader->GetCodeHeader();
 	FUniformBufferGatherInfo UBGatherInfo;
 	DescriptorSetLayoutInfo.ProcessBindingsForStage(VK_SHADER_STAGE_COMPUTE_BIT, ShaderStage::Compute, CSHeader, UBGatherInfo);
-	DescriptorSetLayoutInfo.FinalizeBindings<true>(*Device, UBGatherInfo, TArrayView<FRHISamplerState*>(), Shader->UsesBindless());
+	DescriptorSetLayoutInfo.FinalizeBindings<true>(*Device, UBGatherInfo, Shader->UsesBindless());
 	FVulkanLayout* Layout = FindOrAddLayout(DescriptorSetLayoutInfo, false, Shader->UsesBindless());
 	checkSlow(!Layout->IsGfxLayout());
 	Pipeline->Layout = Layout;
