Add a unique ID to command list names, to ease debugging.

#rnx
#tests Opened a D3D12 capture in RenderDoc and checked that the command lists have unique names

#Synced-CL 50863983

diff --git a/Engine/Source/Runtime/D3D12RHI/Private/D3D12CommandList.cpp b/Engine/Source/Runtime/D3D12RHI/Private/D3D12CommandList.cpp
--- a/Engine/Source/Runtime/D3D12RHI/Private/D3D12CommandList.cpp
+++ b/Engine/Source/Runtime/D3D12RHI/Private/D3D12CommandList.cpp
@@ -4,6 +4,10 @@
 #include "D3D12RHIPrivate.h"
 #include "RHIValidation.h"
 
+#if RHI_USE_RESOURCE_DEBUG_NAME
+static std::atomic_int32_t GNextCommandListID{1};
+#endif
+
 void FD3D12CommandList::UpdateResidency(const FD3D12Resource* Resource)
 {
 #if ENABLE_RESIDENCY_MANAGEMENT
@@ -177,7 +181,8 @@ FD3D12CommandList::FD3D12CommandList(FD3D12CommandAllocator* CommandAllocator, F
 #endif
 
 #if RHI_USE_RESOURCE_DEBUG_NAME
-	const FString Name = FString::Printf(TEXT("FD3D12CommandList (GPU %d)"), Device->GetGPUIndex());
+	const int32 ID = GNextCommandListID++;
+	const FString Name = FString::Printf(TEXT("FD3D12CommandList %d (GPU %d)"), ID, Device->GetGPUIndex());
 	SetD3D12ObjectName(Interfaces.CommandList, GetData(Name));
 #endif
 
diff --git a/Engine/Source/Runtime/VulkanRHI/Private/VulkanCommandBuffer.cpp b/Engine/Source/Runtime/VulkanRHI/Private/VulkanCommandBuffer.cpp
--- a/Engine/Source/Runtime/VulkanRHI/Private/VulkanCommandBuffer.cpp
+++ b/Engine/Source/Runtime/VulkanRHI/Private/VulkanCommandBuffer.cpp
@@ -14,9 +14,11 @@
 #include "VulkanBindlessDescriptorManager.h"
 #include "VulkanPipeline.h"
 
-
 #define CMD_BUFFER_TIME_TO_WAIT_BEFORE_DELETING		10
 
+#if VULKAN_USE_DEBUG_NAMES
+static std::atomic_int32_t GNextCommandBufferID{1};
+#endif
 
 FVulkanCommandBuffer::FVulkanCommandBuffer(FVulkanDevice& InDevice, FVulkanCommandBufferPool& InCommandBufferPool)
 	: Device(InDevice)
@@ -60,6 +62,11 @@ void FVulkanCommandBuffer::AllocMemory()
 	FMemory::Memzero(CurrentGraphicsShaderObjects);
 
 	INC_DWORD_STAT(STAT_VulkanNumCmdBuffers);
+
+#if VULKAN_USE_DEBUG_NAMES
+	int32 ID = GNextCommandBufferID++;
+	VULKAN_SET_DEBUG_NAME(Device, VK_OBJECT_TYPE_COMMAND_BUFFER, CommandBufferHandle, TEXT("FVulkanCommandBuffer %d"), ID);
+#endif
 }
 
 FVulkanCommandBuffer::~FVulkanCommandBuffer()
diff --git a/Engine/Source/Runtime/VulkanRHI/Private/VulkanDevice.h b/Engine/Source/Runtime/VulkanRHI/Private/VulkanDevice.h
--- a/Engine/Source/Runtime/VulkanRHI/Private/VulkanDevice.h
+++ b/Engine/Source/Runtime/VulkanRHI/Private/VulkanDevice.h
@@ -28,7 +28,6 @@ struct FVulkanParallelRenderPassInfo;
 class FVulkanFenceManager;
 class FVulkanShaderObjectManager;
 
-#define VULKAN_USE_DEBUG_NAMES 1
 #if VULKAN_USE_DEBUG_NAMES
 #define VULKAN_SET_DEBUG_NAME(Device, Type, Handle, Format, ...) (Device).VulkanSetObjectName(Type, (uint64)Handle, *FString::Printf(Format, __VA_ARGS__))
 #else
diff --git a/Engine/Source/Runtime/VulkanRHI/Public/VulkanConfiguration.h b/Engine/Source/Runtime/VulkanRHI/Public/VulkanConfiguration.h
--- a/Engine/Source/Runtime/VulkanRHI/Public/VulkanConfiguration.h
+++ b/Engine/Source/Runtime/VulkanRHI/Public/VulkanConfiguration.h
@@ -247,6 +247,9 @@ VULKANRHI_API DECLARE_LOG_CATEGORY_EXTERN(LogVulkanRHI, Log, All);
 #	define VK_FLAGS_TO_STRING(Type, Value) *FString::Printf(TEXT("%u"), (uint32)Value)
 #endif
 
+#ifndef VULKAN_USE_DEBUG_NAMES
+#	define VULKAN_USE_DEBUG_NAMES 1
+#endif
 
 namespace VulkanRHI
 {
