VulkanRHI: Remove functions that are unused or became unnecessary.

#rnx

#Synced-CL 51094475

diff --git a/Engine/Source/Runtime/VulkanRHI/Private/VulkanDescriptorSets.cpp b/Engine/Source/Runtime/VulkanRHI/Private/VulkanDescriptorSets.cpp
--- a/Engine/Source/Runtime/VulkanRHI/Private/VulkanDescriptorSets.cpp
+++ b/Engine/Source/Runtime/VulkanRHI/Private/VulkanDescriptorSets.cpp
@@ -11,11 +11,6 @@ TAutoConsoleVariable<int32> GDynamicGlobalUBs(
 	ECVF_ReadOnly | ECVF_RenderThreadSafe
 );
 
-void FVulkanDescriptorSetsLayoutInfo::ProcessBindingsForStage(VkShaderStageFlagBits StageFlags, ShaderStage::EStage DescSetStage, const FVulkanShaderHeader& CodeHeader, FUniformBufferGatherInfo& OutUBGatherInfo) const
-{
-	OutUBGatherInfo.CodeHeaders[DescSetStage] = &CodeHeader;
-}
-
 template<bool bIsCompute>
 void FVulkanDescriptorSetsLayoutInfo::FinalizeBindings(const FVulkanDevice& Device, const FUniformBufferGatherInfo& UBGatherInfo, bool bUsesBindless)
 {
diff --git a/Engine/Source/Runtime/VulkanRHI/Private/VulkanDescriptorSets.h b/Engine/Source/Runtime/VulkanRHI/Private/VulkanDescriptorSets.h
--- a/Engine/Source/Runtime/VulkanRHI/Private/VulkanDescriptorSets.h
+++ b/Engine/Source/Runtime/VulkanRHI/Private/VulkanDescriptorSets.h
@@ -102,8 +102,6 @@ public:
 		return SetLayouts;
 	}
 
-	void ProcessBindingsForStage(VkShaderStageFlagBits StageFlags, ShaderStage::EStage DescSetStage, const FVulkanShaderHeader& CodeHeader, FUniformBufferGatherInfo& OutUBGatherInfo) const;
-
 	template<bool bIsCompute>
 	void FinalizeBindings(const FVulkanDevice& Device, const FUniformBufferGatherInfo& UBGatherInfo, bool bUsesBindless);
 
diff --git a/Engine/Source/Runtime/VulkanRHI/Private/VulkanLayout.h b/Engine/Source/Runtime/VulkanRHI/Private/VulkanLayout.h
--- a/Engine/Source/Runtime/VulkanRHI/Private/VulkanLayout.h
+++ b/Engine/Source/Runtime/VulkanRHI/Private/VulkanLayout.h
@@ -45,23 +45,6 @@ protected:
 	const bool bIsGfxLayout;
 	const bool bUsesBindless;
 
-	template <bool bIsCompute>
-	void FinalizeBindings(const FUniformBufferGatherInfo& UBGatherInfo)
-	{
-		// Setting descriptor is only allowed prior to compiling the layout
-		check(DescriptorSetLayout.GetHandles().Num() == 0);
-
-		DescriptorSetLayout.FinalizeBindings<bIsCompute>(UBGatherInfo);
-	}
-
-	void ProcessBindingsForStage(VkShaderStageFlagBits StageFlags, ShaderStage::EStage DescSet, const FVulkanShaderHeader& CodeHeader, FUniformBufferGatherInfo& OutUBGatherInfo) const
-	{
-		// Setting descriptor is only allowed prior to compiling the layout
-		check(DescriptorSetLayout.GetHandles().Num() == 0);
-
-		DescriptorSetLayout.ProcessBindingsForStage(StageFlags, DescSet, CodeHeader, OutUBGatherInfo);
-	}
-
 	void Compile(FVulkanDescriptorSetLayoutMap& DSetLayoutMap);
 
 	friend class FVulkanComputePipeline;
diff --git a/Engine/Source/Runtime/VulkanRHI/Private/VulkanPipeline.cpp b/Engine/Source/Runtime/VulkanRHI/Private/VulkanPipeline.cpp
--- a/Engine/Source/Runtime/VulkanRHI/Private/VulkanPipeline.cpp
+++ b/Engine/Source/Runtime/VulkanRHI/Private/VulkanPipeline.cpp
@@ -2346,7 +2346,7 @@ void FVulkanPipelineStateCacheManager::CreateGfxEntry(const FGraphicsPipelineSta
 			if (Shader)
 			{
 				const FVulkanShaderHeader& Header = Shader->GetCodeHeader();
-				DescriptorSetLayoutInfo.ProcessBindingsForStage(StageFlag, Stage, Header, UBGatherInfo);
+				UBGatherInfo.CodeHeaders[Stage] = &Header;
 				NumActiveShaders++;
 				if (Shader->UsesBindless())
 				{
@@ -2870,7 +2870,7 @@ FVulkanComputePipeline* FVulkanPipelineStateCacheManager::CreateComputePipeline(
 	FVulkanDescriptorSetsLayoutInfo DescriptorSetLayoutInfo;
 	const FVulkanShaderHeader& CSHeader = Shader->GetCodeHeader();
 	FUniformBufferGatherInfo UBGatherInfo;
-	DescriptorSetLayoutInfo.ProcessBindingsForStage(VK_SHADER_STAGE_COMPUTE_BIT, ShaderStage::Compute, CSHeader, UBGatherInfo);
+	UBGatherInfo.CodeHeaders[ShaderStage::Compute] = &CSHeader;
 	DescriptorSetLayoutInfo.FinalizeBindings<true>(*Device, UBGatherInfo, Shader->UsesBindless());
 	FVulkanLayout* Layout = FindOrAddLayout(DescriptorSetLayoutInfo, false, Shader->UsesBindless());
 	checkSlow(!Layout->IsGfxLayout());
