VulkanRHI: Fix warning message for image count

#rnx

#Synced-CL 51214690

#Synced-CL 51214690

diff --git a/Engine/Source/Runtime/VulkanRHI/Private/VulkanSwapChain.cpp b/Engine/Source/Runtime/VulkanRHI/Private/VulkanSwapChain.cpp
--- a/Engine/Source/Runtime/VulkanRHI/Private/VulkanSwapChain.cpp
+++ b/Engine/Source/Runtime/VulkanRHI/Private/VulkanSwapChain.cpp
@@ -218,10 +218,6 @@ FVulkanSwapChain* FVulkanSwapChain::Create(
 		UE_LOG(LogVulkanRHI, Warning, TEXT("vkGetPhysicalDeviceSurfaceFormatsKHR returned 0 formats! Handling as surface lost."));
 		VERIFYVULKANSURFACERESULT(VK_ERROR_SURFACE_LOST_KHR);
 	}
-	else if (InOutDesiredNumBackBuffers && (NumFormats < *InOutDesiredNumBackBuffers))
-	{
-		UE_LOG(LogVulkanRHI, Warning, TEXT("vkGetPhysicalDeviceSurfaceFormatsKHR returned %u formats, expected %u!"), NumFormats, *InOutDesiredNumBackBuffers);
-	}
 
 	TArray<VkSurfaceFormatKHR> Formats;
 	Formats.AddZeroed(NumFormats);
@@ -506,6 +502,12 @@ FVulkanSwapChain* FVulkanSwapChain::Create(
 	}
 
 	// 0 means no limit, so use the requested number
+	if (((*InOutDesiredNumBackBuffers) < SurfProperties.minImageCount) || 
+		((SurfProperties.maxImageCount > 0) && ((*InOutDesiredNumBackBuffers) > SurfProperties.maxImageCount)))
+	{
+		UE_LOG(LogVulkanRHI, Warning, TEXT("VkSurfaceCapabilitiesKHR returned [%u,%u] image count, requested %u"), 
+			SurfProperties.minImageCount, SurfProperties.maxImageCount, *InOutDesiredNumBackBuffers);
+	}
 	const uint32 DesiredNumBuffers = SurfProperties.maxImageCount > 0 ? FMath::Clamp(*InOutDesiredNumBackBuffers, SurfProperties.minImageCount, SurfProperties.maxImageCount) : *InOutDesiredNumBackBuffers;
 
 	const uint32 SizeX = FVulkanPlatform::SupportsQuerySurfaceProperties() ? (SurfProperties.currentExtent.width == 0xFFFFFFFF ? Width : SurfProperties.currentExtent.width) : Width;
