VulkanRHI: Fix some validation errors caused by less optimized SPIRV with frozen specialization constants resulting in resources still having bindings even if they aren't used.

#rnx

#Synced-CL 51094475

diff --git a/Engine/Source/Runtime/VulkanRHI/Private/VulkanPipelineState.cpp b/Engine/Source/Runtime/VulkanRHI/Private/VulkanPipelineState.cpp
--- a/Engine/Source/Runtime/VulkanRHI/Private/VulkanPipelineState.cpp
+++ b/Engine/Source/Runtime/VulkanRHI/Private/VulkanPipelineState.cpp
@@ -132,7 +132,10 @@ void FVulkanCommonPipelineDescriptorState::CreateDescriptorWriteInfos()
 			CurrentAccelerationStructures,
 			DefaultSampler, DefaultImageView);
 
-		DSWriter[Set].SetExcludedResources(GetExcludedResources(Set));
+		if (!bUseBindless)
+		{
+			DSWriter[Set].SetExcludedResources(GetExcludedResources(Set));
+		}
 
 		TotalNumDynamicOffsets += NumDynamicOffsets;
 
@@ -179,6 +182,70 @@ static inline VulkanRHI::FVulkanAllocation UpdatePackedUniformBuffers(const FPac
 }
 
 
+static void SetDefaultDescriptorsForMissingBuffers(FVulkanDescriptorSetWriter& DescriptorSetWriter, FRHIShader* Shader, const uint32 SpecializationHash)
+{
+	// Grab defaults from system resources that we can use on descriptors we know won't be accessed
+	const FVulkanView::FStructuredBufferView* DefaultStructureBufferUAV = nullptr;
+	if (GEmptyStructuredBufferWithUAV && GEmptyStructuredBufferWithUAV->UnorderedAccessViewRHI.IsValid())
+	{
+		DefaultStructureBufferUAV = &ResourceCast(GEmptyStructuredBufferWithUAV->UnorderedAccessViewRHI.GetReference())->GetStructuredBufferView();
+	}
+	const FVulkanView::FTypedBufferView* DefaultTypedBufferSRV = nullptr;
+	if (GEmptyVertexBufferWithUAV && GEmptyVertexBufferWithUAV->ShaderResourceViewRHI.IsValid())
+	{
+		DefaultTypedBufferSRV = &ResourceCast(GEmptyVertexBufferWithUAV->ShaderResourceViewRHI.GetReference())->GetTypedBufferView();
+	}
+	const FVulkanView::FTypedBufferView* DefaultTypedBufferUAV = nullptr;
+	if (GEmptyVertexBufferWithUAV && GEmptyVertexBufferWithUAV->UnorderedAccessViewRHI.IsValid())
+	{
+		DefaultTypedBufferUAV = &ResourceCast(GEmptyVertexBufferWithUAV->UnorderedAccessViewRHI.GetReference())->GetTypedBufferView();
+	}
+
+	TArrayView<VkWriteDescriptorSet> DescriptorWrites(const_cast<VkWriteDescriptorSet*>(DescriptorSetWriter.GetWriteDescriptors()), DescriptorSetWriter.GetNumWrites());
+	for (VkWriteDescriptorSet& Write : DescriptorWrites)
+	{
+#if RHI_INCLUDE_SHADER_DEBUG_DATA
+		const bool bIsExcludedFromSpecialization = Shader->IsExcludedFromSpecialization(Write.dstBinding, SpecializationHash);
+#else
+		constexpr bool bIsExcludedFromSpecialization = true;
+#endif
+
+		switch (Write.descriptorType)
+		{
+		case VK_DESCRIPTOR_TYPE_STORAGE_BUFFER:
+			if (DefaultStructureBufferUAV && (Write.pBufferInfo->buffer == VK_NULL_HANDLE))
+			{
+				checkf(bIsExcludedFromSpecialization, TEXT("Missing a storage buffer binding that isn't excluded in this specialization!"));
+				VkDescriptorBufferInfo* BufferInfo = const_cast<VkDescriptorBufferInfo*>(Write.pBufferInfo);
+				BufferInfo->buffer = DefaultStructureBufferUAV->Buffer;
+				BufferInfo->offset = DefaultStructureBufferUAV->Offset;
+				BufferInfo->range = DefaultStructureBufferUAV->Size;
+			}
+			break;
+		case VK_DESCRIPTOR_TYPE_STORAGE_TEXEL_BUFFER:
+			if (DefaultTypedBufferUAV && (Write.pTexelBufferView == nullptr))
+			{
+				checkf(bIsExcludedFromSpecialization, TEXT("Missing a storage texel buffer binding that isn't excluded in this specialization!"));
+				VkWriteDescriptorSet& BufferInfo = const_cast<VkWriteDescriptorSet&>(Write);
+				Write.pTexelBufferView = &DefaultTypedBufferUAV->View;
+			}
+			break;
+		case VK_DESCRIPTOR_TYPE_UNIFORM_TEXEL_BUFFER:
+			if (DefaultTypedBufferSRV && (Write.pTexelBufferView == nullptr))
+			{
+				checkf(bIsExcludedFromSpecialization, TEXT("Missing a uniform texel buffer binding that isn't excluded in this specialization!"));
+				Write.pTexelBufferView = &DefaultTypedBufferSRV->View;
+			}
+			break;
+		default:
+			// Do nothing: textures already have defaults and other types aren't excluded by specialization constants.
+			break;
+		}
+	}
+}
+
+
+
 template<bool bUseDynamicGlobalUBs>
 bool FVulkanComputePipelineDescriptorState::InternalUpdateDescriptorSets(FVulkanCommandListContext& Context)
 {
@@ -224,6 +291,15 @@ bool FVulkanComputePipelineDescriptorState::InternalUpdateDescriptorSets(FVulkan
 		INC_DWORD_STAT(STAT_VulkanNumDescSets);
 		SCOPE_CYCLE_COUNTER(STAT_VulkanVkUpdateDS);
 #endif
+		// Certain layers don't optimize SPIRV enough to know if resources are truly used,
+		// so set defaults to the ones we know will be excluded to silence errors
+		if (ComputePipeline && ComputePipeline->GetComputeShader()->HasSpecializationConstants())
+		{
+			check(ComputePipeline->GetSpecializationConstants().Num());
+			const uint32 SpecializationHash = GetTypeHash(ComputePipeline->GetSpecializationConstants());
+			SetDefaultDescriptorsForMissingBuffers(DSWriter[ShaderStage::Compute], ComputePipeline->GetComputeShader(), SpecializationHash);
+		}
+
 		VulkanRHI::vkUpdateDescriptorSets(Device.GetHandle(), DSWriteContainer.DescriptorWrites.Num(), DSWriteContainer.DescriptorWrites.GetData(), 0, nullptr);
 	}
 
@@ -429,30 +505,46 @@ bool FVulkanGraphicsPipelineDescriptorState::InternalUpdateDescriptorSets(FVulka
 		// Allocate sets based on what changed
 		if (Context.AcquirePoolSetAndDescriptorsIfNeeded(*DescriptorSetsLayout, bNeedsWrite, DescriptorSetHandles.GetData()))
 		{
-			uint32 RemainingSetsMask = UsedSetsMask;
-			uint32 Set = 0;
-			uint32 NumSets = 0;
-			while (RemainingSetsMask)
 			{
-				if (RemainingSetsMask & 1)
+				uint32 RemainingSetsMask = UsedSetsMask;
+				uint32 Set = 0;
+				uint32 NumSets = 0;
+				while (RemainingSetsMask)
 				{
-					const VkDescriptorSet DescriptorSet = DescriptorSetHandles[Set];
-					DSWriter[Set].SetDescriptorSet(DescriptorSet);
+					if (RemainingSetsMask & 1)
+					{
+						const VkDescriptorSet DescriptorSet = DescriptorSetHandles[Set];
+						DSWriter[Set].SetDescriptorSet(DescriptorSet);
 #if VULKAN_VALIDATE_DESCRIPTORS_WRITTEN
-					DSWriter[Set].CheckAllWritten();
+						DSWriter[Set].CheckAllWritten();
 #endif
-					++NumSets;
+						++NumSets;
+					}
+
+					++Set;
+					RemainingSetsMask >>= 1;
 				}
+#if VULKAN_ENABLE_AGGRESSIVE_STATS
+				INC_DWORD_STAT_BY(STAT_VulkanNumUpdateDescriptors, DSWriteContainer.DescriptorWrites.Num());
+				INC_DWORD_STAT_BY(STAT_VulkanNumDescSets, NumSets);
+#endif
+			}
 
-				++Set;
-				RemainingSetsMask >>= 1;
+			// Certain layers don't optimize SPIRV enough to know if resources are truly used,
+			// so set defaults to the ones we know will be excluded to silence errors
+			for (uint32 Stage = 0; Stage < MaxNumSets; ++Stage)
+			{
+				FRHIShader* Shader = GfxPipeline->GetShader(GetFrequencyForGfxStage((ShaderStage::EStage)Stage));
+				if (Shader && Shader->HasSpecializationConstants())
+				{
+					const uint32 SpecializationHash = GetTypeHash(GfxPipeline->SpecializationConstants[Stage]);
+					SetDefaultDescriptorsForMissingBuffers(DSWriter[Stage], Shader, SpecializationHash);
+				}
 			}
 
-	#if VULKAN_ENABLE_AGGRESSIVE_STATS
-			INC_DWORD_STAT_BY(STAT_VulkanNumUpdateDescriptors, DSWriteContainer.DescriptorWrites.Num());
-			INC_DWORD_STAT_BY(STAT_VulkanNumDescSets, NumSets);
-			SCOPE_CYCLE_COUNTER(STAT_VulkanVkUpdateDS);
-	#endif
+#if VULKAN_ENABLE_AGGRESSIVE_STATS
+				SCOPE_CYCLE_COUNTER(STAT_VulkanVkUpdateDS);
+#endif
 			VulkanRHI::vkUpdateDescriptorSets(Device.GetHandle(), DSWriteContainer.DescriptorWrites.Num(), DSWriteContainer.DescriptorWrites.GetData(), 0, nullptr);
 
 			bIsResourcesDirty = false;
