VulkanRHI: Add external semaphore extensions by default in Windows and Linux.  Always load the external memory extensions.  (often used by plugins)

#rnx

#Synced-CL 51237690

diff --git a/Engine/Source/Runtime/VulkanRHI/Private/Linux/VulkanLinuxPlatform.cpp b/Engine/Source/Runtime/VulkanRHI/Private/Linux/VulkanLinuxPlatform.cpp
--- a/Engine/Source/Runtime/VulkanRHI/Private/Linux/VulkanLinuxPlatform.cpp
+++ b/Engine/Source/Runtime/VulkanRHI/Private/Linux/VulkanLinuxPlatform.cpp
@@ -350,8 +350,8 @@ void FVulkanLinuxPlatform::GetInstanceExtensions(FVulkanInstanceExtensionArray&
 
 void FVulkanLinuxPlatform::GetDeviceExtensions(FVulkanDevice* Device, FVulkanDeviceExtensionArray& OutExtensions)
 {
-	// Manually activated extensions
-	OutExtensions.Add(MakeUnique<FVulkanDeviceExtension>(Device, VK_KHR_EXTERNAL_MEMORY_FD_EXTENSION_NAME, VULKAN_EXTENSION_ENABLED, VULKAN_EXTENSION_NOT_PROMOTED, nullptr, FVulkanExtensionBase::ManuallyActivate));
+	OutExtensions.Add(MakeUnique<FVulkanDeviceExtension>(Device, VK_KHR_EXTERNAL_SEMAPHORE_FD_EXTENSION_NAME, VULKAN_EXTENSION_ENABLED, VULKAN_EXTENSION_NOT_PROMOTED));
+	OutExtensions.Add(MakeUnique<FVulkanDeviceExtension>(Device, VK_KHR_EXTERNAL_MEMORY_FD_EXTENSION_NAME, VULKAN_EXTENSION_ENABLED, VULKAN_EXTENSION_NOT_PROMOTED));
 }
 
 void FVulkanLinuxPlatform::CreateSurface(FVulkanPlatformWindowContext& WindowContext, VkInstance Instance, VkSurfaceKHR* OutSurface)
diff --git a/Engine/Source/Runtime/VulkanRHI/Private/Windows/VulkanWindowsPlatform.cpp b/Engine/Source/Runtime/VulkanRHI/Private/Windows/VulkanWindowsPlatform.cpp
--- a/Engine/Source/Runtime/VulkanRHI/Private/Windows/VulkanWindowsPlatform.cpp
+++ b/Engine/Source/Runtime/VulkanRHI/Private/Windows/VulkanWindowsPlatform.cpp
@@ -170,9 +170,8 @@ void FVulkanWindowsPlatform::GetDeviceExtensions(FVulkanDevice* Device, FVulkanD
 	OutExtensions.Add(MakeUnique<FVulkanDeviceExtension>(Device, VK_EXT_FULL_SCREEN_EXCLUSIVE_EXTENSION_NAME, VULKAN_SUPPORTS_FULLSCREEN_EXCLUSIVE, 
 														VULKAN_EXTENSION_NOT_PROMOTED, DEVICE_EXT_FLAG_SETTER(HasEXTFullscreenExclusive)));
 
-	// Manually activated extensions
-	OutExtensions.Add(MakeUnique<FVulkanDeviceExtension>(Device, VK_KHR_EXTERNAL_MEMORY_WIN32_EXTENSION_NAME, VULKAN_EXTENSION_ENABLED,
-														VULKAN_EXTENSION_NOT_PROMOTED, nullptr, FVulkanExtensionBase::ManuallyActivate));
+	OutExtensions.Add(MakeUnique<FVulkanDeviceExtension>(Device, VK_KHR_EXTERNAL_SEMAPHORE_WIN32_EXTENSION_NAME, VULKAN_EXTENSION_ENABLED, VULKAN_EXTENSION_NOT_PROMOTED));
+	OutExtensions.Add(MakeUnique<FVulkanDeviceExtension>(Device, VK_KHR_EXTERNAL_MEMORY_WIN32_EXTENSION_NAME, VULKAN_EXTENSION_ENABLED, VULKAN_EXTENSION_NOT_PROMOTED));
 }
 
 void FVulkanWindowsPlatform::CreateSurface(FVulkanPlatformWindowContext& WindowContext, VkInstance Instance, VkSurfaceKHR* OutSurface)
