VulkanRHI: Add a flag for protected memory

#rnx

#Synced-CL 51303845

diff --git a/Engine/Source/Runtime/VulkanRHI/Private/VulkanMemory.cpp b/Engine/Source/Runtime/VulkanRHI/Private/VulkanMemory.cpp
--- a/Engine/Source/Runtime/VulkanRHI/Private/VulkanMemory.cpp
+++ b/Engine/Source/Runtime/VulkanRHI/Private/VulkanMemory.cpp
@@ -2758,7 +2758,11 @@ namespace VulkanRHI
 		{
 			MemFlags = VK_MEMORY_PROPERTY_DEVICE_LOCAL_BIT;
 
-			if (bHasUnifiedMemory)
+			if (EnumHasAnyFlags(AllocFlags, EVulkanAllocationFlags::Protected))
+			{
+				MemFlags |= VK_MEMORY_PROPERTY_PROTECTED_BIT;
+			}
+			else if (bHasUnifiedMemory)
 			{
 				MemFlags |= (VK_MEMORY_PROPERTY_HOST_VISIBLE_BIT | VK_MEMORY_PROPERTY_HOST_COHERENT_BIT);
 			}
diff --git a/Engine/Source/Runtime/VulkanRHI/Private/VulkanMemory.h b/Engine/Source/Runtime/VulkanRHI/Private/VulkanMemory.h
--- a/Engine/Source/Runtime/VulkanRHI/Private/VulkanMemory.h
+++ b/Engine/Source/Runtime/VulkanRHI/Private/VulkanMemory.h
@@ -705,9 +705,10 @@ namespace VulkanRHI
 		Dedicated    = 0x0008,	  // Will not share a memory block with other resources
 		External     = 0x0010,    // To be used with VK_KHR_external_memory
 		Memoryless   = 0x0020,	  // Will use LAZILY_ALLOCATED
+		Protected    = 0x0040,	  // Will use PROTECTED
 
-		NoError      = 0x0040,    // OOM is not fatal, return an invalid allocation
-		AutoBind     = 0x0080,	  // Will automatically bind the allocation to the supplied VkBuffer/VkImage, avoids an extra lock to bind separately
+		NoError      = 0x0080,    // OOM is not fatal, return an invalid allocation
+		AutoBind     = 0x0100,	  // Will automatically bind the allocation to the supplied VkBuffer/VkImage, avoids an extra lock to bind separately
 	};
 	ENUM_CLASS_FLAGS(EVulkanAllocationFlags);
 
