VulkanRHI: remove MemoryRequirement argument from AllocateDedicatedImageMemory since it queries requirements in the function.  Always include the code for the function.

#rnx

#Synced-CL 51303845

diff --git a/Engine/Source/Runtime/VulkanRHI/Private/VulkanMemory.cpp b/Engine/Source/Runtime/VulkanRHI/Private/VulkanMemory.cpp
--- a/Engine/Source/Runtime/VulkanRHI/Private/VulkanMemory.cpp
+++ b/Engine/Source/Runtime/VulkanRHI/Private/VulkanMemory.cpp
@@ -3140,9 +3140,8 @@ namespace VulkanRHI
 		return OutAllocation.IsValid();
 	}
 
-	bool FMemoryManager::AllocateDedicatedImageMemory(FVulkanAllocation& OutAllocation, FVulkanEvictable* AllocationOwner, VkImage Image, const VkMemoryRequirements& MemoryReqs, VkMemoryPropertyFlags MemoryPropertyFlags, EVulkanAllocationMetaType MetaType, bool bExternal, const char* File, uint32 Line)
+	bool FMemoryManager::AllocateDedicatedImageMemory(FVulkanAllocation& OutAllocation, FVulkanEvictable* AllocationOwner, VkImage Image, VkMemoryPropertyFlags MemoryPropertyFlags, EVulkanAllocationMetaType MetaType, bool bExternal, const char* File, uint32 Line)
 	{
-#if VULKAN_SUPPORTS_DEDICATED_ALLOCATION
 		SCOPED_NAMED_EVENT(FVulkanMemoryManager_AllocateDedicatedImageMemory, FColor::Cyan);
 		VkImageMemoryRequirementsInfo2KHR ImageMemoryReqs2;
 		ZeroVulkanStruct(ImageMemoryReqs2, VK_STRUCTURE_TYPE_IMAGE_MEMORY_REQUIREMENTS_INFO_2_KHR);
@@ -3156,8 +3155,9 @@ namespace VulkanRHI
 		MemoryReqs2.pNext = &DedMemoryReqs;
 
 		VulkanRHI::vkGetImageMemoryRequirements2(Device.GetHandle(), &ImageMemoryReqs2, &MemoryReqs2);
+		const VkMemoryRequirements& MemoryReqs = MemoryReqs2.memoryRequirements;
 
-		bool bUseDedicated = DedMemoryReqs.prefersDedicatedAllocation != VK_FALSE || DedMemoryReqs.requiresDedicatedAllocation != VK_FALSE;
+		const bool bUseDedicated = DedMemoryReqs.prefersDedicatedAllocation != VK_FALSE || DedMemoryReqs.requiresDedicatedAllocation != VK_FALSE;
 		if (bUseDedicated)
 		{
 			uint32 TypeIndex = 0;
@@ -3189,10 +3189,6 @@ namespace VulkanRHI
 		{
 			return AllocateImageMemory(OutAllocation, AllocationOwner, MemoryReqs, MemoryPropertyFlags, MetaType, bExternal, File, Line);
 		}
-#else
-		checkNoEntry();
-		return false;
-#endif
 	}
 
 
diff --git a/Engine/Source/Runtime/VulkanRHI/Private/VulkanMemory.h b/Engine/Source/Runtime/VulkanRHI/Private/VulkanMemory.h
--- a/Engine/Source/Runtime/VulkanRHI/Private/VulkanMemory.h
+++ b/Engine/Source/Runtime/VulkanRHI/Private/VulkanMemory.h
@@ -741,7 +741,7 @@ namespace VulkanRHI
 		// Legacy calls
 		bool AllocateBufferPooled(FVulkanAllocation& Allocation, FVulkanEvictable* AllocationOwner, uint32 Size, uint32 MinAlignment, VkBufferUsageFlags BufferUsageFlags, VkMemoryPropertyFlags MemoryPropertyFlags, EVulkanAllocationMetaType MetaType, const char* File, uint32 Line);
 		bool AllocateImageMemory(FVulkanAllocation& Allocation, FVulkanEvictable* AllocationOwner, const VkMemoryRequirements& MemoryReqs, VkMemoryPropertyFlags MemoryPropertyFlags, EVulkanAllocationMetaType MetaType, bool bExternal, const char* File, uint32 Line);
-		bool AllocateDedicatedImageMemory(FVulkanAllocation& Allocation, FVulkanEvictable* AllocationOwner, VkImage Image, const VkMemoryRequirements& MemoryReqs, VkMemoryPropertyFlags MemoryPropertyFlags, EVulkanAllocationMetaType MetaType, bool bExternal, const char* File, uint32 Line);
+		bool AllocateDedicatedImageMemory(FVulkanAllocation& Allocation, FVulkanEvictable* AllocationOwner, VkImage Image, VkMemoryPropertyFlags MemoryPropertyFlags, EVulkanAllocationMetaType MetaType, bool bExternal, const char* File, uint32 Line);
 	private:
 		bool AllocateBufferMemory(FVulkanAllocation& Allocation, FVulkanEvictable* AllocationOwner, const VkMemoryRequirements& MemoryReqs, VkMemoryPropertyFlags MemoryPropertyFlags, EVulkanAllocationMetaType MetaType, bool bExternal, bool bForceSeparateAllocation, const char* File, uint32 Line);
 
diff --git a/Engine/Source/Runtime/VulkanRHI/Private/VulkanTexture.cpp b/Engine/Source/Runtime/VulkanRHI/Private/VulkanTexture.cpp
--- a/Engine/Source/Runtime/VulkanRHI/Private/VulkanTexture.cpp
+++ b/Engine/Source/Runtime/VulkanRHI/Private/VulkanTexture.cpp
@@ -1915,7 +1915,7 @@ FVulkanTexture::FVulkanTexture(FVulkanDevice& InDevice, const FRHITextureCreateD
 			VkDeviceSize SizeToBeConsideredForDedicated = 12 * 1024 * 1024;
 			if ((bRenderTarget || MemoryRequirements.size >= SizeToBeConsideredForDedicated) && !bMemoryless && GVulkanEnableDedicatedImageMemory)
 			{
-				if (!InDevice.GetMemoryManager().AllocateDedicatedImageMemory(Allocation, this, Image, MemoryRequirements, MemoryFlags, MetaType, bExternal, __FILE__, __LINE__))
+				if (!InDevice.GetMemoryManager().AllocateDedicatedImageMemory(Allocation, this, Image, MemoryFlags, MetaType, bExternal, __FILE__, __LINE__))
 				{
 					checkNoEntry();
 				}
