VulkanRHI: Add missing const on certain accessors.  Add missing accessor for memory type.

#rnx

#Synced-CL 51303845

diff --git a/Engine/Source/Runtime/VulkanRHI/Private/VulkanMemory.h b/Engine/Source/Runtime/VulkanRHI/Private/VulkanMemory.h
--- a/Engine/Source/Runtime/VulkanRHI/Private/VulkanMemory.h
+++ b/Engine/Source/Runtime/VulkanRHI/Private/VulkanMemory.h
@@ -519,12 +519,12 @@ namespace VulkanRHI
 		bool TryAllocate2(FVulkanAllocation& OutAllocation, FVulkanEvictable* Owner, uint32 InSize, uint32 InAlignment, EVulkanAllocationMetaType InMetaType, const char* File, uint32 Line);
 		void Free(FVulkanAllocation& Allocation);
 
-		inline uint32 GetAlignment() const
+		uint32 GetAlignment() const
 		{
 			return Alignment;
 		}
 
-		inline void* GetMappedPointer()
+		void* GetMappedPointer()
 		{
 			return MemoryAllocation->GetMappedPointer();
 		}
@@ -534,7 +534,12 @@ namespace VulkanRHI
 		uint32 GetMaxSize() const { return MaxSize; }
 		uint32 GetUsedSize() const { return UsedSize; }
 
-		inline uint32 GetHandleId() const
+		uint32 GetMemoryTypeIndex() const
+		{
+			return MemoryTypeIndex;
+		}
+
+		uint32 GetHandleId() const
 		{
 			return BufferId;
 		}
@@ -547,16 +552,18 @@ namespace VulkanRHI
 		TArrayView<uint32> GetMemoryUsed();
 		uint32 GetNumSubAllocations();
 
-		uint8 GetSubresourceAllocatorFlags()
+		uint8 GetSubresourceAllocatorFlags() const
 		{
 			return SubresourceAllocatorFlags;
 		}
-		uint32 GetId()
+		uint32 GetId() const
 		{
 			return BufferId;
 		}
-		bool GetIsDefragging() { return bIsDefragging; }
-
+		bool GetIsDefragging() const
+		{
+			return bIsDefragging;
+		}
 
 	protected:
 		void SetIsDefragging(bool bInIsDefragging){ bIsDefragging = bInIsDefragging; }
@@ -572,7 +579,7 @@ namespace VulkanRHI
 
 		EVulkanAllocationType Type;
 		FMemoryManager* Owner;
-		uint32 MemoryTypeIndex;
+		const uint32 MemoryTypeIndex;
 		VkMemoryPropertyFlags MemoryPropertyFlags;
 		FDeviceMemoryAllocation* MemoryAllocation;
 		uint32 MaxSize;
@@ -661,7 +668,7 @@ namespace VulkanRHI
 		TArray<FVulkanPageSizeBucket, TFixedAllocator<MAX_BUCKETS> > PageSizeBuckets;
 		uint32 GetPageSize();
 		FMemoryManager* Owner;
-		uint16 MemoryTypeIndex;
+		const uint16 MemoryTypeIndex;
 		const uint16 HeapIndex;
 
 		bool bIsHostCachedSupported;
