VulkanRHI:  Add VkMemoryPropertyFlags to FVulkanRHIAllocationInfo.

#rnx

#Synced-CL 51303845

diff --git a/Engine/Source/Runtime/VulkanRHI/Private/VulkanMemory.cpp b/Engine/Source/Runtime/VulkanRHI/Private/VulkanMemory.cpp
--- a/Engine/Source/Runtime/VulkanRHI/Private/VulkanMemory.cpp
+++ b/Engine/Source/Runtime/VulkanRHI/Private/VulkanMemory.cpp
@@ -3964,6 +3964,13 @@ namespace VulkanRHI
 		}
 		return 0;
 	}
+	VkMemoryPropertyFlags FVulkanAllocation::GetMemoryPropertyFlags(FVulkanDevice* Device) const
+	{
+		const VkPhysicalDeviceMemoryProperties& MemoryProperties = Device->GetDeviceMemoryManager().GetMemoryProperties();
+		const uint32 MemoryTypeIndex = GetSubresourceAllocator(Device)->GetMemoryTypeIndex();
+		check(MemoryTypeIndex < MemoryProperties.memoryTypeCount);
+		return MemoryProperties.memoryTypes[MemoryTypeIndex].propertyFlags;
+	}
 
 	void FVulkanSubresourceAllocator::FreeInternalData(int32 Index)
 	{
diff --git a/Engine/Source/Runtime/VulkanRHI/Private/VulkanMemory.h b/Engine/Source/Runtime/VulkanRHI/Private/VulkanMemory.h
--- a/Engine/Source/Runtime/VulkanRHI/Private/VulkanMemory.h
+++ b/Engine/Source/Runtime/VulkanRHI/Private/VulkanMemory.h
@@ -457,6 +457,7 @@ namespace VulkanRHI
 		uint32 GetBufferAlignment(FVulkanDevice* Device) const;
 		VkDeviceMemory GetDeviceMemoryHandle(FVulkanDevice* Device) const;
 		FVulkanSubresourceAllocator* GetSubresourceAllocator(FVulkanDevice* Device) const;
+		VkMemoryPropertyFlags GetMemoryPropertyFlags(FVulkanDevice* Device) const;
 		void BindBuffer(FVulkanDevice* Device, VkBuffer Buffer);
 		void BindImage(FVulkanDevice* Device, VkImage Image);
 	};
diff --git a/Engine/Source/Runtime/VulkanRHI/Private/VulkanRHI.cpp b/Engine/Source/Runtime/VulkanRHI/Private/VulkanRHI.cpp
--- a/Engine/Source/Runtime/VulkanRHI/Private/VulkanRHI.cpp
+++ b/Engine/Source/Runtime/VulkanRHI/Private/VulkanRHI.cpp
@@ -1580,6 +1580,7 @@ FVulkanRHIAllocationInfo FVulkanDynamicRHI::RHIGetAllocationInfo(FRHITexture* In
 	NewInfo.Handle = VulkanTexture->GetAllocationHandle();
 	NewInfo.Offset = VulkanTexture->GetAllocationOffset();
 	NewInfo.Size = VulkanTexture->GetMemorySize();
+	NewInfo.Flags = VulkanTexture->MemProps;
 
 	return NewInfo;
 }
@@ -1619,6 +1620,7 @@ FVulkanRHIAllocationInfo FVulkanDynamicRHI::RHIGetAllocationInfo(FRHIBuffer* InB
 	NewInfo.Handle = Allocation.GetDeviceMemoryHandle(GetDevice());
 	NewInfo.Offset = Allocation.Offset;
 	NewInfo.Size = Allocation.Size;
+	NewInfo.Flags = Allocation.GetMemoryPropertyFlags(GetDevice());
 
 	return NewInfo;
 }
diff --git a/Engine/Source/Runtime/VulkanRHI/Public/IVulkanDynamicRHI.h b/Engine/Source/Runtime/VulkanRHI/Public/IVulkanDynamicRHI.h
--- a/Engine/Source/Runtime/VulkanRHI/Public/IVulkanDynamicRHI.h
+++ b/Engine/Source/Runtime/VulkanRHI/Public/IVulkanDynamicRHI.h
@@ -10,6 +10,7 @@ struct FVulkanRHIAllocationInfo
 	VkDeviceMemory Handle{};
 	uint64 Offset{};
 	uint64 Size{};
+	VkMemoryPropertyFlags Flags{};
 };
 
 struct FVulkanRHIImageViewInfo
