VulkanRHI: Convert extension logging messages to UE_LOGF.

#rnx

#Synced-CL 51379665

diff --git a/Engine/Source/Runtime/VulkanRHI/Private/VulkanExtensions.cpp b/Engine/Source/Runtime/VulkanRHI/Private/VulkanExtensions.cpp
--- a/Engine/Source/Runtime/VulkanRHI/Private/VulkanExtensions.cpp
+++ b/Engine/Source/Runtime/VulkanRHI/Private/VulkanExtensions.cpp
@@ -462,9 +462,9 @@ public:
 	virtual void PostPhysicalDeviceProperties(FOptionalVulkanDeviceExtensions& ExtensionFlags) override final
 	{
 		const VkPhysicalDeviceDriverPropertiesKHR& PhysicalDeviceDriverProperties = GetDeviceExtensionProperties().PhysicalDeviceDriverProperties;
-		UE_LOG(LogVulkanRHI, Display, TEXT("Vulkan Physical Device Driver Properties:"));
-		UE_LOG(LogVulkanRHI, Display, TEXT("- driverName: %s"), ANSI_TO_TCHAR(PhysicalDeviceDriverProperties.driverName));
-		UE_LOG(LogVulkanRHI, Display, TEXT("- driverInfo: %s"), ANSI_TO_TCHAR(PhysicalDeviceDriverProperties.driverInfo));
+		UE_LOGF(LogVulkanRHI, Display, "Vulkan Physical Device Driver Properties:");
+		UE_LOGF(LogVulkanRHI, Display, "- driverName: %s", PhysicalDeviceDriverProperties.driverName);
+		UE_LOGF(LogVulkanRHI, Display, "- driverInfo: %s", PhysicalDeviceDriverProperties.driverInfo);
 	}
 };
 
@@ -805,7 +805,7 @@ public:
 
 			// todo: We don't currently care much about the other properties here, but at some point in the future we probably will.
 
-			UE_LOG(LogVulkanRHI, Verbose, TEXT("Image-based Variable Rate Shading supported via KHRFragmentShadingRate extension. Selected VRS tile size %u by %u pixels per VRS image texel."), GRHIVariableRateShadingImageTileMinWidth, GRHIVariableRateShadingImageTileMinHeight);
+			UE_LOGF(LogVulkanRHI, Verbose, "Image-based Variable Rate Shading supported via KHRFragmentShadingRate extension. Selected VRS tile size %u by %u pixels per VRS image texel.", GRHIVariableRateShadingImageTileMinWidth, GRHIVariableRateShadingImageTileMinHeight);
 		}
 	}
 
@@ -887,7 +887,7 @@ public:
 			GRHIVariableRateShadingImageDataType = VRSImage_Fractional;
 			GRHIVariableRateShadingImageFormat = PF_R8G8;
 
-			UE_LOG(LogVulkanRHI, Display, TEXT("Image-based Variable Rate Shading supported via EXTFragmentDensityMap extension. Selected VRS tile size %u by %u pixels per VRS image texel."), GRHIVariableRateShadingImageTileMinWidth, GRHIVariableRateShadingImageTileMinHeight);
+			UE_LOGF(LogVulkanRHI, Display, "Image-based Variable Rate Shading supported via EXTFragmentDensityMap extension. Selected VRS tile size %u by %u pixels per VRS image texel.", GRHIVariableRateShadingImageTileMinWidth, GRHIVariableRateShadingImageTileMinHeight);
 		}
 	}
 
@@ -1791,7 +1791,7 @@ public:
 
 			if (DescriptorBufferProperties.bufferlessPushDescriptors == VK_FALSE)
 			{
-				UE_LOG(LogVulkanRHI, Display, TEXT("Disabling Vulkan Descriptor Buffers because bufferless isn't supported."));
+				UE_LOGF(LogVulkanRHI, Display, "Disabling Vulkan Descriptor Buffers because bufferless isn't supported.");
 
 				ExtensionFlags.HasEXTDescriptorBuffer = false;
 			}
@@ -1804,13 +1804,13 @@ public:
 		{
 			const VkPhysicalDeviceDescriptorBufferPropertiesEXT& DescriptorBufferProperties = GetDeviceExtensionProperties().DescriptorBufferProps;
 
-			UE_LOG(LogVulkanRHI, Display, TEXT("Enabling Vulkan Descriptor Buffers with: ")
-				TEXT("allowSamplerImageViewPostSubmitCreation=%u, maxDescriptorBufferBindings=%u, ")
-				TEXT("maxSamplerDescriptorBufferBindings=%u, maxResourceDescriptorBufferBindings=%u, ")
-				TEXT("samplerDescriptorBufferAddressSpaceSize=%llu, resourceDescriptorBufferAddressSpaceSize=%llu, ")
-				TEXT("maxSamplerDescriptorBufferRange=%llu, maxResourceDescriptorBufferRange=%llu, ")
-				TEXT("descriptorBufferAddressSpaceSize=%llu, descriptorBufferOffsetAlignment=%llu, ")
-				TEXT("samplerDescriptorSize=%llu"),
+			UE_LOGF(LogVulkanRHI, Display, "Enabling Vulkan Descriptor Buffers with: "
+				"allowSamplerImageViewPostSubmitCreation=%u, maxDescriptorBufferBindings=%u, "
+				"maxSamplerDescriptorBufferBindings=%u, maxResourceDescriptorBufferBindings=%u, "
+				"samplerDescriptorBufferAddressSpaceSize=%llu, resourceDescriptorBufferAddressSpaceSize=%llu, "
+				"maxSamplerDescriptorBufferRange=%llu, maxResourceDescriptorBufferRange=%llu, "
+				"descriptorBufferAddressSpaceSize=%llu, descriptorBufferOffsetAlignment=%llu, "
+				"samplerDescriptorSize=%llu",
 				DescriptorBufferProperties.allowSamplerImageViewPostSubmitCreation, DescriptorBufferProperties.maxDescriptorBufferBindings,
 				DescriptorBufferProperties.maxSamplerDescriptorBufferBindings, DescriptorBufferProperties.maxResourceDescriptorBufferBindings,
 				DescriptorBufferProperties.samplerDescriptorBufferAddressSpaceSize, DescriptorBufferProperties.resourceDescriptorBufferAddressSpaceSize,
@@ -2471,7 +2471,7 @@ public:
 		ExtensionFlags.HasEXTHostImageCopy = bRequirementsPassed;
 		if (bRequirementsPassed)
 		{
-			UE_LOG(LogVulkanRHI, Display, TEXT("Host Image Copy Enabled (SrcCount=%d, DstCount=%d)"), 
+			UE_LOGF(LogVulkanRHI, Display, "Host Image Copy Enabled (SrcCount=%d, DstCount=%d)", 
 				HostImageCopyProperties.copySrcLayoutCount, HostImageCopyProperties.copyDstLayoutCount);
 		}
 	}
@@ -2527,7 +2527,7 @@ public:
 		ExtensionFlags.HasEXTDeviceGeneratedCommands = bRequirementsPassed;
 		if (bRequirementsPassed)
 		{
-			UE_LOG(LogVulkanRHI, Display, TEXT("Device Generated Commands Enabled (PipelineCount=%d, SequenceCount=%d)"),
+			UE_LOGF(LogVulkanRHI, Display, "Device Generated Commands Enabled (PipelineCount=%d, SequenceCount=%d)",
 				DeviceGeneratedCommandsProperties.maxIndirectPipelineCount, DeviceGeneratedCommandsProperties.maxIndirectSequenceCount);
 
 			GRHISupportsShaderBundleDispatch = true;
@@ -2583,7 +2583,7 @@ public:
 	{
 		if (bRequirementsPassed)
 		{
-			UE_LOG(LogVulkanRHI, Display, TEXT("Robustness2 Enabled: Buffer=%d (Align SB/UB=%u/%u), Image=%d, NullDescriptor=%d"),
+			UE_LOGF(LogVulkanRHI, Display, "Robustness2 Enabled: Buffer=%d (Align SB/UB=%u/%u), Image=%d, NullDescriptor=%d",
 				Robustness2Features.robustBufferAccess2, 
 				(uint32)Robustness2Properties.robustStorageBufferAccessSizeAlignment, (uint32)Robustness2Properties.robustUniformBufferAccessSizeAlignment,
 				Robustness2Features.robustImageAccess2, Robustness2Features.nullDescriptor);
@@ -2605,10 +2605,10 @@ public:
 
 
 template <typename ExtensionType>
-static void FlagExtensionSupport(const TArray<VkExtensionProperties>& ExtensionProperties, TArray<TUniquePtr<ExtensionType>>& UEExtensions, uint32 ApiVersion, const TCHAR* ExtensionTypeName)
+static void FlagExtensionSupport(const TArray<VkExtensionProperties>& ExtensionProperties, TArray<TUniquePtr<ExtensionType>>& UEExtensions, uint32 ApiVersion, const ANSICHAR* ExtensionTypeName)
 {
 	// Flag the extension support
-	UE_LOG(LogVulkanRHI, Display, TEXT("Found %d available %s extensions :"), ExtensionProperties.Num(), ExtensionTypeName);
+	UE_LOGF(LogVulkanRHI, Display, "Found %d available %s extensions :", ExtensionProperties.Num(), ExtensionTypeName);
 	for (const VkExtensionProperties& Extension : ExtensionProperties)
 	{
 		const int32 ExtensionIndex = ExtensionType::FindExtension(UEExtensions, Extension.extensionName);
@@ -2622,7 +2622,7 @@ static void FlagExtensionSupport(const TArray<VkExtensionProperties>& ExtensionP
 			bIsCore = UEExtensions[ExtensionIndex]->SetCore(ApiVersion);
 		}
 
-		UE_LOG(LogVulkanRHI, Display, TEXT("  %s %s"), bIsCore ? TEXT("*") : bFound ? TEXT("+") : TEXT("-"), ANSI_TO_TCHAR(Extension.extensionName));
+		UE_LOGF(LogVulkanRHI, Display, "  %s %s", bIsCore ? "*" : bFound ? "+" : "-", Extension.extensionName);
 	}
 }
 
@@ -2733,14 +2733,14 @@ FVulkanDeviceExtensionArray FVulkanDeviceExtension::GetUESupportedDeviceExtensio
 	FVulkanPlatform::GetDeviceExtensions(InDevice, OutUEDeviceExtensions);
 
 	// Helper function to go through a list of extensions and activate them (or add them)
-	auto ActivateExternalExtensions = [&](TArray<const ANSICHAR*> Extensions, const TCHAR* Requester) {
+	auto ActivateExternalExtensions = [&](TArray<const ANSICHAR*> Extensions, const ANSICHAR* Requester) {
 		for (const ANSICHAR* ExtensionName : Extensions)
 		{
 			const int32 ExtensionIndex = FindExtension(OutUEDeviceExtensions, ExtensionName);
 			if (ExtensionIndex == INDEX_NONE)
 			{
 				OutUEDeviceExtensions.Add(MakeUnique<FVulkanDeviceExtension>(InDevice, ExtensionName, VULKAN_EXTENSION_ENABLED));
-				UE_LOG(LogVulkanRHI, Warning, TEXT("%s requested device extension [%s] isn't part of the engine's core extension list. Adding it on-the-fly..."), Requester, ANSI_TO_TCHAR(ExtensionName));
+				UE_LOGF(LogVulkanRHI, Warning, "%s requested device extension [%s] isn't part of the engine's core extension list. Adding it on-the-fly...", Requester, ExtensionName);
 			}
 			else
 			{
@@ -2759,15 +2759,15 @@ FVulkanDeviceExtensionArray FVulkanDeviceExtension::GetUESupportedDeviceExtensio
 		{
 			TArray<const ANSICHAR*> HMDExtensions;
 			FVulkanDynamicRHI::HMDVulkanExtensions->GetVulkanDeviceExtensionsRequired(InDevice->GetPhysicalHandle(), HMDExtensions);
-			ActivateExternalExtensions(HMDExtensions, TEXT("HMD"));
+			ActivateExternalExtensions(HMDExtensions, "HMD");
 		}
 	}
 
 	// Add extensions added outside the RHI (eg plugins)
-	ActivateExternalExtensions(ExternalExtensions, TEXT("Externally"));
+	ActivateExternalExtensions(ExternalExtensions, "Externally");
 
 	// Now that all the extensions are listed, update their support flags
-	FlagExtensionSupport(GetDriverSupportedDeviceExtensions(InDevice->GetPhysicalHandle()), OutUEDeviceExtensions, ApiVersion, TEXT("device"));
+	FlagExtensionSupport(GetDriverSupportedDeviceExtensions(InDevice->GetPhysicalHandle()), OutUEDeviceExtensions, ApiVersion, "device");
 
 	#undef ADD_SIMPLE_EXTENSION
 	#undef ADD_EXTERNAL_EXTENSION
@@ -2872,14 +2872,14 @@ FVulkanInstanceExtensionArray FVulkanInstanceExtension::GetUESupportedInstanceEx
 	FVulkanPlatform::GetInstanceExtensions(OutUEInstanceExtensions);
 
 	// Helper function to go through a list of extensions and activate them (or add them)
-	auto ActivateExternalExtensions = [&] (TArray<const ANSICHAR*> Extensions, const TCHAR* Requester) {
+	auto ActivateExternalExtensions = [&] (TArray<const ANSICHAR*> Extensions, const ANSICHAR* Requester) {
 		for (const ANSICHAR* ExtensionName : Extensions)
 		{
 			const int32 ExtensionIndex = FindExtension(OutUEInstanceExtensions, ExtensionName);
 			if (ExtensionIndex == INDEX_NONE)
 			{
 				OutUEInstanceExtensions.Add(MakeUnique<FVulkanInstanceExtension>(ExtensionName, VULKAN_EXTENSION_ENABLED));
-				UE_LOG(LogVulkanRHI, Warning, TEXT("%s requested instance extension [%s] isn't part of the engine's core extension list. Adding it on-the-fly..."), Requester, ANSI_TO_TCHAR(ExtensionName));
+				UE_LOGF(LogVulkanRHI, Warning, "%s requested instance extension [%s] isn't part of the engine's core extension list. Adding it on-the-fly...", Requester, ExtensionName);
 			}
 			else
 			{
@@ -2899,15 +2899,15 @@ FVulkanInstanceExtensionArray FVulkanInstanceExtension::GetUESupportedInstanceEx
 		{
 			TArray<const ANSICHAR*> HMDExtensions;
 			FVulkanDynamicRHI::HMDVulkanExtensions->GetVulkanInstanceExtensionsRequired(HMDExtensions);
-			ActivateExternalExtensions(HMDExtensions, TEXT("HMD"));
+			ActivateExternalExtensions(HMDExtensions, "HMD");
 		}
 	}
 
 	// Add extensions added outside the RHI (eg plugins)
-	ActivateExternalExtensions(ExternalExtensions, TEXT("Externally"));
+	ActivateExternalExtensions(ExternalExtensions, "Externally");
 
 	// Now that all the extensions are listed, update their support flags
-	FlagExtensionSupport(GetDriverSupportedInstanceExtensions(), OutUEInstanceExtensions, ApiVersion, TEXT("instance"));
+	FlagExtensionSupport(GetDriverSupportedInstanceExtensions(), OutUEInstanceExtensions, ApiVersion, "instance");
 
 	return OutUEInstanceExtensions;
 }
