VulkanRHI: Load VK_KHR_shader_untyped_pointers and VK_EXT_descriptor_heap extensions if present.

#rnx

#Synced-CL 51414939

diff --git a/Engine/Source/Runtime/VulkanRHI/Private/VulkanCommandWrappers.h b/Engine/Source/Runtime/VulkanRHI/Private/VulkanCommandWrappers.h
--- a/Engine/Source/Runtime/VulkanRHI/Private/VulkanCommandWrappers.h
+++ b/Engine/Source/Runtime/VulkanRHI/Private/VulkanCommandWrappers.h
@@ -319,6 +319,13 @@ struct FWrapLayer
 	static void GetSwapchainTimeDomainPropertiesEXT(VkResult Result, VkDevice Device, VkSwapchainKHR Swapchain, VkSwapchainTimeDomainPropertiesEXT* SwapchainTimeDomainProperties, uint64_t* TimeDomainsCounter) VULKAN_LAYER_BODY
 	static void GetSwapchainTimingPropertiesEXT(VkResult Result, VkDevice Device, VkSwapchainKHR Swapchain, VkSwapchainTimingPropertiesEXT* SwapchainTimingProperties, uint64_t* SwapchainTimingPropertiesCounter) VULKAN_LAYER_BODY
 	static void SetSwapchainPresentTimingQueueSizeEXT(VkResult Result, VkDevice Device, VkSwapchainKHR Swapchain, uint32_t Size) VULKAN_LAYER_BODY
+
+	static void CmdBindResourceHeapEXT(VkResult Result, VkCommandBuffer CommandBuffer, const VkBindHeapInfoEXT* BindInfo) VULKAN_LAYER_BODY
+	static void CmdBindSamplerHeapEXT(VkResult Result, VkCommandBuffer CommandBuffer, const VkBindHeapInfoEXT* BindInfo) VULKAN_LAYER_BODY
+	static void CmdPushDataEXT(VkResult Result, VkCommandBuffer CommandBuffer, const VkPushDataInfoEXT* PushDataInfo) VULKAN_LAYER_BODY
+	static void GetPhysicalDeviceDescriptorSizeEXT(VkDeviceSize OutSize, VkPhysicalDevice PhysicalDevice, VkDescriptorType DescriptorType) VULKAN_LAYER_BODY
+	static void WriteResourceDescriptorsEXT(VkResult Result, VkDevice Device, uint32_t ResourceCount, const VkResourceDescriptorInfoEXT* Resources, const VkHostAddressRangeEXT* Descriptors) VULKAN_LAYER_BODY
+	static void WriteSamplerDescriptorsEXT(VkResult Result, VkDevice Device, uint32_t SamplerCount, const VkSamplerCreateInfo* Samplers, const VkHostAddressRangeEXT* Descriptors) VULKAN_LAYER_BODY
 };
 
 
@@ -2264,6 +2271,53 @@ namespace VulkanRHI
 		return Result;
 	}
 
+	static FORCEINLINE_DEBUGGABLE void vkCmdBindResourceHeapEXT(VkCommandBuffer CommandBuffer, const VkBindHeapInfoEXT* BindInfo)
+	{
+		FWrapLayer::CmdBindResourceHeapEXT(VK_RESULT_MAX_ENUM, CommandBuffer, BindInfo);
+		VULKANAPINAMESPACE::vkCmdBindResourceHeapEXT(CommandBuffer, BindInfo);
+		FWrapLayer::CmdBindResourceHeapEXT(VK_SUCCESS, CommandBuffer, BindInfo);
+	}
+
+	static FORCEINLINE_DEBUGGABLE void vkCmdBindSamplerHeapEXT(VkCommandBuffer CommandBuffer, const VkBindHeapInfoEXT* BindInfo)
+	{
+		FWrapLayer::CmdBindSamplerHeapEXT(VK_RESULT_MAX_ENUM, CommandBuffer, BindInfo);
+		VULKANAPINAMESPACE::vkCmdBindSamplerHeapEXT(CommandBuffer, BindInfo);
+		FWrapLayer::CmdBindSamplerHeapEXT(VK_SUCCESS, CommandBuffer, BindInfo);
+	}
+
+	static FORCEINLINE_DEBUGGABLE void vkCmdPushDataEXT(VkCommandBuffer CommandBuffer, const VkPushDataInfoEXT* PushDataInfo)
+	{
+		FWrapLayer::CmdPushDataEXT(VK_RESULT_MAX_ENUM, CommandBuffer, PushDataInfo);
+		VULKANAPINAMESPACE::vkCmdPushDataEXT(CommandBuffer, PushDataInfo);
+		FWrapLayer::CmdPushDataEXT(VK_SUCCESS, CommandBuffer, PushDataInfo);
+	}
+
+	static FORCEINLINE_DEBUGGABLE VkDeviceSize vkGetPhysicalDeviceDescriptorSizeEXT(VkPhysicalDevice PhysicalDevice, VkDescriptorType DescriptorType)
+	{
+		FWrapLayer::GetPhysicalDeviceDescriptorSizeEXT(0, PhysicalDevice, DescriptorType);
+		const VkDeviceSize OutSize = VULKANAPINAMESPACE::vkGetPhysicalDeviceDescriptorSizeEXT(PhysicalDevice, DescriptorType);
+		FWrapLayer::GetPhysicalDeviceDescriptorSizeEXT(OutSize, PhysicalDevice, DescriptorType);
+		return OutSize;
+	}
+
+	static FORCEINLINE_DEBUGGABLE VkResult vkWriteResourceDescriptorsEXT(VkDevice Device, uint32_t ResourceCount, const VkResourceDescriptorInfoEXT* Resources, const VkHostAddressRangeEXT* Descriptors)
+	{
+		FWrapLayer::WriteResourceDescriptorsEXT(VK_RESULT_MAX_ENUM, Device, ResourceCount, Resources, Descriptors);
+		const VkResult Result = VULKANAPINAMESPACE::vkWriteResourceDescriptorsEXT(Device, ResourceCount, Resources, Descriptors);
+		FWrapLayer::WriteResourceDescriptorsEXT(Result, Device, ResourceCount, Resources, Descriptors);
+		return Result;
+	}
+
+	static FORCEINLINE_DEBUGGABLE VkResult vkWriteSamplerDescriptorsEXT(VkDevice Device, uint32_t SamplerCount, const VkSamplerCreateInfo* Samplers, const VkHostAddressRangeEXT* Descriptors)
+	{
+		FWrapLayer::WriteSamplerDescriptorsEXT(VK_RESULT_MAX_ENUM, Device, SamplerCount, Samplers, Descriptors);
+		const VkResult Result = VULKANAPINAMESPACE::vkWriteSamplerDescriptorsEXT(Device, SamplerCount, Samplers, Descriptors);
+		FWrapLayer::WriteSamplerDescriptorsEXT(Result, Device, SamplerCount, Samplers, Descriptors);
+		return Result;
+	}
+
+
+
 
 #if VULKAN_ENABLE_IMAGE_TRACKING_LAYER
 	void BindDebugLabelName(VkImage Image, const TCHAR* Name);
diff --git a/Engine/Source/Runtime/VulkanRHI/Private/VulkanDebug.cpp b/Engine/Source/Runtime/VulkanRHI/Private/VulkanDebug.cpp
--- a/Engine/Source/Runtime/VulkanRHI/Private/VulkanDebug.cpp
+++ b/Engine/Source/Runtime/VulkanRHI/Private/VulkanDebug.cpp
@@ -5215,6 +5215,71 @@ void FWrapLayer::SetSwapchainPresentTimingQueueSizeEXT(VkResult Result, VkDevice
 	}
 }
 
+void FWrapLayer::CmdBindResourceHeapEXT(VkResult Result, VkCommandBuffer CommandBuffer, const VkBindHeapInfoEXT* BindInfo)
+{
+	if (Result == VK_RESULT_MAX_ENUM)
+	{
+#if VULKAN_ENABLE_DUMP_LAYER
+		PrintfBeginResult(FString::Printf(TEXT("CmdBindResourceHeapEXT()")));
+#endif
+	}
+}
+
+void FWrapLayer::CmdBindSamplerHeapEXT(VkResult Result, VkCommandBuffer CommandBuffer, const VkBindHeapInfoEXT* BindInfo)
+{
+	if (Result == VK_RESULT_MAX_ENUM)
+	{
+#if VULKAN_ENABLE_DUMP_LAYER
+		PrintfBeginResult(FString::Printf(TEXT("CmdBindSamplerHeapEXT()")));
+#endif
+	}
+}
+
+void FWrapLayer::CmdPushDataEXT(VkResult Result, VkCommandBuffer CommandBuffer, const VkPushDataInfoEXT* PushDataInfo)
+{
+	if (Result == VK_RESULT_MAX_ENUM)
+	{
+#if VULKAN_ENABLE_DUMP_LAYER
+		PrintfBeginResult(FString::Printf(TEXT("CmdPushDataEXT()")));
+#endif
+	}
+}
+
+void FWrapLayer::GetPhysicalDeviceDescriptorSizeEXT(vkDeviceSize OutSize, VkPhysicalDevice PhysicalDevice, VkDescriptorType DescriptorType)
+{
+	if (OutSize > 0)
+	{
+#if VULKAN_ENABLE_DUMP_LAYER
+		PrintfBeginResult(FString::Printf(TEXT("GetPhysicalDeviceDescriptorSizeEXT()")));
+#endif
+	}
+}
+
+void FWrapLayer::WriteResourceDescriptorsEXT(VkResult Result, VkDevice Device, uint32_t ResourceCount, const VkResourceDescriptorInfoEXT* Resources, const VkHostAddressRangeEXT* Descriptors)
+{
+	if (Result == VK_RESULT_MAX_ENUM)
+	{
+#if VULKAN_ENABLE_DUMP_LAYER
+		PrintfBeginResult(FString::Printf(TEXT("WriteResourceDescriptorsEXT()")));
+#endif
+	}
+}
+
+void FWrapLayer::WriteSamplerDescriptorsEXT(VkResult Result, VkDevice Device, uint32_t SamplerCount, const VkSamplerCreateInfo* Samplers, const VkHostAddressRangeEXT* Descriptors)
+{
+	if (Result == VK_RESULT_MAX_ENUM)
+	{
+#if VULKAN_ENABLE_DUMP_LAYER
+		PrintfBeginResult(FString::Printf(TEXT("WriteSamplerDescriptorsEXT()")));
+#endif
+	}
+}
+
+
+
+
+
+
 
 #if VULKAN_ENABLE_IMAGE_TRACKING_LAYER
 namespace VulkanRHI
diff --git a/Engine/Source/Runtime/VulkanRHI/Private/VulkanDevice.h b/Engine/Source/Runtime/VulkanRHI/Private/VulkanDevice.h
--- a/Engine/Source/Runtime/VulkanRHI/Private/VulkanDevice.h
+++ b/Engine/Source/Runtime/VulkanRHI/Private/VulkanDevice.h
@@ -59,6 +59,7 @@ struct FOptionalVulkanDeviceExtensions
 			uint64 HasDeferredHostOperations : 1;
 			uint64 HasEXTCalibratedTimestamps : 1;
 			uint64 HasEXTDescriptorBuffer : 1;
+			uint64 HasEXTDescriptorHeap : 1;
 			uint64 HasEXTDeviceFault : 1;
 			uint64 HasEXTMeshShader : 1;
 			uint64 HasEXTToolingInfo : 1;
@@ -165,6 +166,7 @@ struct FOptionalVulkanDeviceExtensionProperties
 	VkPhysicalDeviceMaintenance7PropertiesKHR PhysicalDeviceMaintenance7Properties;
 
 	VkPhysicalDeviceDescriptorBufferPropertiesEXT DescriptorBufferProps;
+	VkPhysicalDeviceDescriptorHeapPropertiesEXT DescriptorHeapProps;
 	VkPhysicalDeviceSubgroupSizeControlPropertiesEXT SubgroupSizeControlProperties;
 
 	VkPhysicalDeviceAccelerationStructurePropertiesKHR AccelerationStructureProps;
diff --git a/Engine/Source/Runtime/VulkanRHI/Private/VulkanExtensions.cpp b/Engine/Source/Runtime/VulkanRHI/Private/VulkanExtensions.cpp
--- a/Engine/Source/Runtime/VulkanRHI/Private/VulkanExtensions.cpp
+++ b/Engine/Source/Runtime/VulkanRHI/Private/VulkanExtensions.cpp
@@ -1826,6 +1826,100 @@ private:
 	VkPhysicalDeviceDescriptorBufferFeaturesEXT DescriptorBufferFeatures;
 };
 
+
+// ***** VK_KHR_shader_untyped_pointers
+class FVulkanKHRShaderUntypedPointersExtension : public FVulkanDeviceExtension
+{
+public:
+
+	FVulkanKHRShaderUntypedPointersExtension(FVulkanDevice* InDevice)
+		: FVulkanDeviceExtension(InDevice, VK_KHR_SHADER_UNTYPED_POINTERS_EXTENSION_NAME, VULKAN_EXTENSION_ENABLED)
+	{
+	}
+
+	virtual void PrePhysicalDeviceFeatures(VkPhysicalDeviceFeatures2KHR& PhysicalDeviceFeatures2) override final
+	{
+		ZeroVulkanStruct(ShaderUntypedPointersFeatures, VK_STRUCTURE_TYPE_PHYSICAL_DEVICE_SHADER_UNTYPED_POINTERS_FEATURES_KHR);
+		AddToPNext(PhysicalDeviceFeatures2, ShaderUntypedPointersFeatures);
+	}
+
+	virtual void PostPhysicalDeviceFeatures(FOptionalVulkanDeviceExtensions& ExtensionFlags) override final
+	{
+		bRequirementsPassed = ShaderUntypedPointersFeatures.shaderUntypedPointers == VK_TRUE;
+	}
+
+	virtual void PreCreateDevice(VkDeviceCreateInfo& DeviceCreateInfo) override final
+	{
+		if (bRequirementsPassed)
+		{
+			AddToPNext(DeviceCreateInfo, ShaderUntypedPointersFeatures);
+		}
+	}
+
+	VkPhysicalDeviceShaderUntypedPointersFeaturesKHR ShaderUntypedPointersFeatures;
+};
+
+
+// ***** VK_EXT_descriptor_heap
+class FVulkanEXTDescriptorHeapExtension : public FVulkanDeviceExtension
+{
+public:
+
+	FVulkanEXTDescriptorHeapExtension(FVulkanDevice* InDevice)
+		: FVulkanDeviceExtension(InDevice, VK_EXT_DESCRIPTOR_HEAP_EXTENSION_NAME, VULKAN_EXTENSION_ENABLED)
+	{
+	}
+
+	virtual void PrePhysicalDeviceFeatures(VkPhysicalDeviceFeatures2KHR& PhysicalDeviceFeatures2) override final
+	{
+		ZeroVulkanStruct(DescriptorHeapFeatures, VK_STRUCTURE_TYPE_PHYSICAL_DEVICE_DESCRIPTOR_HEAP_FEATURES_EXT);
+		AddToPNext(PhysicalDeviceFeatures2, DescriptorHeapFeatures);
+	}
+
+	virtual void PostPhysicalDeviceFeatures(FOptionalVulkanDeviceExtensions& ExtensionFlags) override final
+	{
+		bRequirementsPassed = DescriptorHeapFeatures.descriptorHeap == VK_TRUE;
+		ExtensionFlags.HasEXTDescriptorHeap = bRequirementsPassed;
+	}
+
+	virtual void PrePhysicalDeviceProperties(VkPhysicalDeviceProperties2KHR& PhysicalDeviceProperties2) override final
+	{
+		if (bRequirementsPassed)
+		{
+			VkPhysicalDeviceDescriptorHeapPropertiesEXT& DescriptorHeapProperties = GetDeviceExtensionProperties().DescriptorHeapProps;
+			ZeroVulkanStruct(DescriptorHeapProperties, VK_STRUCTURE_TYPE_PHYSICAL_DEVICE_DESCRIPTOR_HEAP_PROPERTIES_EXT);
+			AddToPNext(PhysicalDeviceProperties2, DescriptorHeapProperties);
+		}
+	}
+
+	virtual void PostPhysicalDeviceProperties(FOptionalVulkanDeviceExtensions& ExtensionFlags) override final
+	{
+		if (bRequirementsPassed)
+		{
+			const VkPhysicalDeviceDescriptorHeapPropertiesEXT& DescriptorBufferProperties = GetDeviceExtensionProperties().DescriptorHeapProps;
+			UE_LOGF(LogVulkanRHI, Display, "Descriptor Heap is supported with following properties: "
+				"maxSamplerHeapSize=%llu, maxResourceHeapSize=%llu, maxPushDataSize=%llu, "
+				"samplerDescriptorAlignment=%llu, imageDescriptorAlignment=%llu, bufferDescriptorAlignment=%llu, "
+				"samplerDescriptorSize=%llu, imageDescriptorSize=%llu, bufferDescriptorSize=%llu, ",
+				DescriptorBufferProperties.maxSamplerHeapSize, DescriptorBufferProperties.maxResourceHeapSize, DescriptorBufferProperties.maxPushDataSize,
+				DescriptorBufferProperties.samplerDescriptorAlignment, DescriptorBufferProperties.imageDescriptorAlignment, DescriptorBufferProperties.bufferDescriptorAlignment,
+				DescriptorBufferProperties.samplerDescriptorSize, DescriptorBufferProperties.imageDescriptorSize, DescriptorBufferProperties.bufferDescriptorSize
+				);
+		}
+	}
+
+	virtual void PreCreateDevice(VkDeviceCreateInfo& DeviceCreateInfo) override final
+	{
+		if (bRequirementsPassed)
+		{
+			AddToPNext(DeviceCreateInfo, DescriptorHeapFeatures);
+		}
+	}
+
+	VkPhysicalDeviceDescriptorHeapFeaturesEXT DescriptorHeapFeatures;
+};
+
+
 // ***** VK_KHR_16bit_storage
 class FVulkanKHR16BitStorageExtension : public FVulkanDeviceExtension
 {
@@ -2714,6 +2808,8 @@ FVulkanDeviceExtensionArray FVulkanDeviceExtension::GetUESupportedDeviceExtensio
 	ADD_CUSTOM_EXTENSION(FVulkanEXTHostImageCopyExtension);
 	ADD_CUSTOM_EXTENSION(FVulkanEXTDeviceGeneratedCommandsExtension);
 	ADD_CUSTOM_EXTENSION(FVulkanEXTRobustness2Extension);
+	ADD_CUSTOM_EXTENSION(FVulkanKHRShaderUntypedPointersExtension);
+	ADD_CUSTOM_EXTENSION(FVulkanEXTDescriptorHeapExtension);
 
 	// Needed for Raytracing
 	ADD_CUSTOM_EXTENSION(FVulkanKHRBufferDeviceAddressExtension);
diff --git a/Engine/Source/Runtime/VulkanRHI/Private/VulkanLoader.h b/Engine/Source/Runtime/VulkanRHI/Private/VulkanLoader.h
--- a/Engine/Source/Runtime/VulkanRHI/Private/VulkanLoader.h
+++ b/Engine/Source/Runtime/VulkanRHI/Private/VulkanLoader.h
@@ -292,7 +292,13 @@
 	EnumMacro(PFN_vkGetPastPresentationTimingEXT, vkGetPastPresentationTimingEXT) \
 	EnumMacro(PFN_vkGetSwapchainTimeDomainPropertiesEXT, vkGetSwapchainTimeDomainPropertiesEXT) \
 	EnumMacro(PFN_vkGetSwapchainTimingPropertiesEXT, vkGetSwapchainTimingPropertiesEXT) \
-	EnumMacro(PFN_vkSetSwapchainPresentTimingQueueSizeEXT, vkSetSwapchainPresentTimingQueueSizeEXT)
+	EnumMacro(PFN_vkSetSwapchainPresentTimingQueueSizeEXT, vkSetSwapchainPresentTimingQueueSizeEXT) \
+	EnumMacro(PFN_vkCmdBindResourceHeapEXT, vkCmdBindResourceHeapEXT) \
+	EnumMacro(PFN_vkCmdBindSamplerHeapEXT, vkCmdBindSamplerHeapEXT) \
+	EnumMacro(PFN_vkCmdPushDataEXT, vkCmdPushDataEXT) \
+	EnumMacro(PFN_vkGetPhysicalDeviceDescriptorSizeEXT, vkGetPhysicalDeviceDescriptorSizeEXT) \
+	EnumMacro(PFN_vkWriteResourceDescriptorsEXT, vkWriteResourceDescriptorsEXT) \
+	EnumMacro(PFN_vkWriteSamplerDescriptorsEXT, vkWriteSamplerDescriptorsEXT)
 
 
 
