VulkanRHI: Clean up legacy defines from when SDK versions weren't guaranteed.

#rnx

#Synced-CL 51414939

diff --git a/Engine/Source/Runtime/VulkanRHI/Private/Windows/VulkanWindowsPlatform.h b/Engine/Source/Runtime/VulkanRHI/Private/Windows/VulkanWindowsPlatform.h
--- a/Engine/Source/Runtime/VulkanRHI/Private/Windows/VulkanWindowsPlatform.h
+++ b/Engine/Source/Runtime/VulkanRHI/Private/Windows/VulkanWindowsPlatform.h
@@ -16,6 +16,7 @@
 #define VULKAN_SIGNAL_UNIMPLEMENTED()				checkf(false, TEXT("Unimplemented vulkan functionality: %s"), StringCast<TCHAR>(__FUNCTION__).Get())
 #define VULKAN_SUPPORTS_SCALAR_BLOCK_LAYOUT			1
 #define VULKAN_SUPPORTS_RAY_TRACING_POSITION_FETCH	1
+#define VULKAN_SUPPORTS_FULLSCREEN_EXCLUSIVE		1
 
 #define UE_VK_API_VERSION							VK_API_VERSION_1_1
 
diff --git a/Engine/Source/Runtime/VulkanRHI/Public/VulkanConfiguration.h b/Engine/Source/Runtime/VulkanRHI/Public/VulkanConfiguration.h
--- a/Engine/Source/Runtime/VulkanRHI/Public/VulkanConfiguration.h
+++ b/Engine/Source/Runtime/VulkanRHI/Public/VulkanConfiguration.h
@@ -109,16 +109,12 @@ struct VkAllocationCallbacks;
 	#define VULKAN_ENABLE_LRU_CACHE								0
 #endif
 
-#ifdef VK_EXT_validation_cache
+#ifndef VULKAN_SUPPORTS_VALIDATION_CACHE
 	#define VULKAN_SUPPORTS_VALIDATION_CACHE					1
-#else
-	#define VULKAN_SUPPORTS_VALIDATION_CACHE					0
 #endif
 
-#ifdef VK_EXT_validation_features
+#ifndef VULKAN_HAS_VALIDATION_FEATURES
 	#define VULKAN_HAS_VALIDATION_FEATURES						1
-#else
-	#define VULKAN_HAS_VALIDATION_FEATURES						0
 #endif
 
 #ifndef VULKAN_SUPPORTS_DEDICATED_ALLOCATION
@@ -163,27 +159,15 @@ struct VkAllocationCallbacks;
 #endif
 
 #ifndef VULKAN_SUPPORTS_QCOM_RENDERPASS_SHADER_RESOLVE
-	#ifdef VK_QCOM_render_pass_shader_resolve
-		#define VULKAN_SUPPORTS_QCOM_RENDERPASS_SHADER_RESOLVE	1
-	#else
-		#define VULKAN_SUPPORTS_QCOM_RENDERPASS_SHADER_RESOLVE	0
-	#endif
+	#define VULKAN_SUPPORTS_QCOM_RENDERPASS_SHADER_RESOLVE		1
 #endif
 
 #ifndef VULKAN_SUPPORTS_FULLSCREEN_EXCLUSIVE
-	#ifdef VK_EXT_full_screen_exclusive
-		#define VULKAN_SUPPORTS_FULLSCREEN_EXCLUSIVE			1
-	#else
-		#define VULKAN_SUPPORTS_FULLSCREEN_EXCLUSIVE			0
-	#endif
+	#define VULKAN_SUPPORTS_FULLSCREEN_EXCLUSIVE				0
 #endif
 
 #ifndef VULKAN_SUPPORTS_TEXTURE_COMPRESSION_ASTC_HDR
-#ifdef VK_EXT_texture_compression_astc_hdr
-#define VULKAN_SUPPORTS_TEXTURE_COMPRESSION_ASTC_HDR			1
-#else
-#define VULKAN_SUPPORTS_TEXTURE_COMPRESSION_ASTC_HDR			0
-#endif
+	#define VULKAN_SUPPORTS_TEXTURE_COMPRESSION_ASTC_HDR		1
 #endif
 
 #ifndef VULKAN_SUPPORTS_ANDROID_EXTERNAL_MEMORY_ANDROID_HARDWARE_BUFFER
@@ -191,33 +175,19 @@ struct VkAllocationCallbacks;
 #endif
 
 #ifndef VULKAN_SUPPORTS_RENDERPASS2
-	#ifdef VK_KHR_create_renderpass2
-		#define VULKAN_SUPPORTS_RENDERPASS2 1
-	#else
-		#define VULKAN_SUPPORTS_RENDERPASS2 0
-	#endif
+	#define VULKAN_SUPPORTS_RENDERPASS2							1
 #endif
 
 #ifndef VULKAN_SUPPORTS_ASTC_DECODE_MODE
-	#ifdef VK_EXT_astc_decode_mode
-		#define VULKAN_SUPPORTS_ASTC_DECODE_MODE				1
-	#else
-		#define VULKAN_SUPPORTS_ASTC_DECODE_MODE				0
-	#endif
+	#define VULKAN_SUPPORTS_ASTC_DECODE_MODE					1
 #endif
 
-#ifdef VK_EXT_shader_viewport_index_layer
-	#define VULKAN_SUPPORTS_SHADER_VIEWPORT_INDEX_LAYER	1
-#else
-	#define VULKAN_SUPPORTS_SHADER_VIEWPORT_INDEX_LAYER	0
+#ifndef VULKAN_SUPPORTS_SHADER_VIEWPORT_INDEX_LAYER
+	#define VULKAN_SUPPORTS_SHADER_VIEWPORT_INDEX_LAYER			1
 #endif
 
 #ifndef VULKAN_SUPPORTS_DESCRIPTOR_INDEXING
-	#ifdef VK_EXT_descriptor_indexing
-		#define VULKAN_SUPPORTS_DESCRIPTOR_INDEXING	1
-	#else
-		#define VULKAN_SUPPORTS_DESCRIPTOR_INDEXING	0
-	#endif
+	#define VULKAN_SUPPORTS_DESCRIPTOR_INDEXING					1
 #endif
 
 #ifndef VULKAN_OBJECT_TRACKING 
@@ -248,7 +218,7 @@ VULKANRHI_API DECLARE_LOG_CATEGORY_EXTERN(LogVulkanRHI, Log, All);
 #endif
 
 #ifndef VULKAN_USE_DEBUG_NAMES
-#	define VULKAN_USE_DEBUG_NAMES 1
+#	define VULKAN_USE_DEBUG_NAMES								1
 #endif
 
 namespace VulkanRHI
