VulkanRHI: change startup logging from UE_LOG to UE_LOGF

#rnx

#Synced-CL 51414939

diff --git a/Engine/Source/Runtime/VulkanRHI/Private/VulkanRHI.cpp b/Engine/Source/Runtime/VulkanRHI/Private/VulkanRHI.cpp
--- a/Engine/Source/Runtime/VulkanRHI/Private/VulkanRHI.cpp
+++ b/Engine/Source/Runtime/VulkanRHI/Private/VulkanRHI.cpp
@@ -39,7 +39,7 @@
 // Use Vulkan Profiles to verify feature level support on startup
 void VulkanProfilePrint(const char* Msg)
 {
-	UE_LOG(LogVulkanRHI, Log, TEXT("   - %s"), ANSI_TO_TCHAR(Msg));
+	UE_LOGF(LogVulkanRHI, Log, "   - %s", Msg);
 }
 #define VP_DEBUG_MESSAGE_CALLBACK VulkanProfilePrint
 #include "vulkan_profiles_ue.h"
@@ -132,9 +132,9 @@ static VkPhysicalDevice SelectPhysicalDevice(VkInstance InInstance)
 	Result = VulkanRHI::vkEnumeratePhysicalDevices(InInstance, &PhysicalDeviceCount, nullptr);
 	if ((Result != VK_SUCCESS) || (PhysicalDeviceCount == 0))
 	{
-		UE_LOG(LogVulkanRHI, Log, 
-			TEXT("SelectPhysicalDevice could not find a compatible Vulkan device or driver (EnumeratePhysicalDevices returned '%s' and %d devices).  ")
-			TEXT("Make sure your video card supports Vulkan and try updating your video driver to a more recent version (proceed with any pending reboots).")
+		UE_LOGF(LogVulkanRHI, Log, 
+			"SelectPhysicalDevice could not find a compatible Vulkan device or driver (EnumeratePhysicalDevices returned '%ls' and %d devices).  "
+			"Make sure your video card supports Vulkan and try updating your video driver to a more recent version (proceed with any pending reboots)."
 			, VK_TYPE_TO_STRING(VkResult, Result), PhysicalDeviceCount
 		);
 
@@ -184,7 +184,7 @@ static VkPhysicalDevice SelectPhysicalDevice(VkInstance InInstance)
 		{
 			if (FMemory::Memcmp(&HmdGraphicsAdapterLuid, &PhysicalDeviceInfos[Index].PhysicalDeviceIDProperties.deviceLUID, VK_LUID_SIZE_KHR) == 0)
 			{
-				UE_LOG(LogVulkanRHI, Log, TEXT("HMD device at index %d of %u being used as default..."), Index, PhysicalDeviceCount);
+				UE_LOGF(LogVulkanRHI, Log, "HMD device at index %d of %u being used as default...", Index, PhysicalDeviceCount);
 				return PhysicalDeviceInfos[Index].PhysicalDevice;
 			}
 		}
@@ -195,22 +195,22 @@ static VkPhysicalDevice SelectPhysicalDevice(VkInstance InInstance)
 	auto* CVarGraphicsAdapter = IConsoleManager::Get().FindTConsoleVariableDataInt(TEXT("r.GraphicsAdapter"));
 	int32 ExplicitAdapterValue = CVarGraphicsAdapter ? CVarGraphicsAdapter->GetValueOnAnyThread() : -1;
 	const bool bUsingCmdLine = FParse::Value(FCommandLine::Get(), TEXT("graphicsadapter="), ExplicitAdapterValue);
-	const TCHAR* GraphicsAdapterOriginTxt = bUsingCmdLine ? TEXT("command line") : TEXT("'r.GraphicsAdapter'");
+	const ANSICHAR* GraphicsAdapterOriginTxt = bUsingCmdLine ? "command line" : "'r.GraphicsAdapter'";
 	if (ExplicitAdapterValue >= 0)  // Use adapter at the specified index
 	{
 		if (ExplicitAdapterValue >= PhysicalDeviceInfos.Num())
 		{
-			UE_LOG(LogVulkanRHI, Warning, TEXT("Tried to use graphics adapter at index %d as specified by %s, but only %d Adapter(s) found. Falling back to first device..."), 
+			UE_LOGF(LogVulkanRHI, Warning, "Tried to use graphics adapter at index %d as specified by %s, but only %d Adapter(s) found. Falling back to first device...", 
 				ExplicitAdapterValue, GraphicsAdapterOriginTxt, PhysicalDeviceInfos.Num());
 			ExplicitAdapterValue = 0;
 		}
 
-		UE_LOG(LogVulkanRHI, Log, TEXT("Using device at index %d of %u as specfified by %s..."), ExplicitAdapterValue, PhysicalDeviceCount, GraphicsAdapterOriginTxt);
+		UE_LOGF(LogVulkanRHI, Log, "Using device at index %d of %u as specfified by %s...", ExplicitAdapterValue, PhysicalDeviceCount, GraphicsAdapterOriginTxt);
 		return PhysicalDeviceInfos[ExplicitAdapterValue].PhysicalDevice;
 	}
 	else if (ExplicitAdapterValue == -2)  // Take the first one that fulfills the criteria
 	{
-		UE_LOG(LogVulkanRHI, Log, TEXT("Using first device (of %u) without any sorting as specfified by %s..."), PhysicalDeviceCount, GraphicsAdapterOriginTxt);
+		UE_LOGF(LogVulkanRHI, Log, "Using first device (of %u) without any sorting as specfified by %s...", PhysicalDeviceCount, GraphicsAdapterOriginTxt);
 		return PhysicalDeviceInfos[0].PhysicalDevice;
 	}
 	else if (ExplicitAdapterValue == -1)  // Favour non integrated because there are usually faster
@@ -240,7 +240,7 @@ static VkPhysicalDevice SelectPhysicalDevice(VkInstance InInstance)
 		{
 			if (RHIConvertToGpuVendorId(PhysicalDeviceInfos[Index].PhysicalDeviceProperties2.properties.vendorID) == PreferredVendor)
 			{
-				UE_LOG(LogVulkanRHI, Log, TEXT("Using preferred vendor device at index %d of %u..."), Index, PhysicalDeviceCount);
+				UE_LOGF(LogVulkanRHI, Log, "Using preferred vendor device at index %d of %u...", Index, PhysicalDeviceCount);
 				return PhysicalDeviceInfos[Index].PhysicalDevice;
 			}
 		}
@@ -258,7 +258,7 @@ static VkPhysicalDevice SelectPhysicalDevice(VkInstance InInstance)
 		return PhysicalDeviceInfos[Index].PhysicalDevice;
 	}
 
-	UE_LOG(LogVulkanRHI, Warning, TEXT("None of the %u devices meet all the criteria!"), PhysicalDeviceCount);
+	UE_LOGF(LogVulkanRHI, Warning, "None of the %u devices meet all the criteria!", PhysicalDeviceCount);
 	return VK_NULL_HANDLE;
 }
 
@@ -275,7 +275,7 @@ static uint32 GetVulkanApiVersionForFeatureLevel(ERHIFeatureLevel::Type FeatureL
 		return minApiVersion;
 	}
 
-	UE_LOG(LogVulkanRHI, Log, TEXT("Using default apiVersion for platform..."));
+	UE_LOGF(LogVulkanRHI, Log, "Using default apiVersion for platform...");
 	return UE_VK_API_VERSION;
 }
 
@@ -286,14 +286,14 @@ static bool CheckVulkanProfile(ERHIFeatureLevel::Type FeatureLevel, bool bRaytra
 
 	if (!FVulkanGenericPlatform::SupportsProfileChecks())
 	{
-		UE_LOG(LogVulkanRHI, Log, TEXT("Skipping Vulkan Profile check for %s:"), *ProfileName);
+		UE_LOGF(LogVulkanRHI, Log, "Skipping Vulkan Profile check for %ls:", *ProfileName);
 		return true;
 	}
 
-	UE_LOG(LogVulkanRHI, Log, TEXT("Starting Vulkan Profile check for %s:"), *ProfileName);
+	UE_LOGF(LogVulkanRHI, Log, "Starting Vulkan Profile check for %ls:", *ProfileName);
 	ON_SCOPE_EXIT
 	{
-		UE_LOG(LogVulkanRHI, Log, TEXT("Vulkan Profile check complete."));
+		UE_LOGF(LogVulkanRHI, Log, "Vulkan Profile check complete.");
 	};
 
 	VpProfileProperties ProfileProperties;
@@ -412,7 +412,7 @@ FDynamicRHI* FVulkanDynamicRHIModule::CreateRHI(ERHIFeatureLevel::Type InRequest
 			check(GMaxTextureSamplers >= (int32)FDataDrivenShaderPlatformInfo::GetMaxSamplers(ShaderPlatformForFeatureLevel[Index]));
 			if (GMaxTextureSamplers < (int32)FDataDrivenShaderPlatformInfo::GetMaxSamplers(ShaderPlatformForFeatureLevel[Index]))
 			{
-				UE_LOG(LogVulkanRHI, Error, TEXT("Shader platform requires at least: %d samplers, device supports: %d."), FDataDrivenShaderPlatformInfo::GetMaxSamplers(ShaderPlatformForFeatureLevel[Index]), GMaxTextureSamplers);
+				UE_LOGF(LogVulkanRHI, Error, "Shader platform requires at least: %d samplers, device supports: %d.", FDataDrivenShaderPlatformInfo::GetMaxSamplers(ShaderPlatformForFeatureLevel[Index]), GMaxTextureSamplers);
 			}
 		}
 	}
@@ -442,7 +442,7 @@ FVulkanDynamicRHI::FVulkanDynamicRHI()
 
 	GRHISupportsPSOPrecaching = CVarAllowVulkanPSOPrecache.GetValueOnAnyThread();
 	GRHISupportsPipelineFileCache = !GRHISupportsPSOPrecaching || CVarEnableVulkanPSOFileCacheWhenPrecachingActive.GetValueOnAnyThread();
-	UE_LOG(LogVulkanRHI, Log, TEXT("Vulkan PSO Precaching = %d, PipelineFileCache = %d"), GRHISupportsPSOPrecaching, GRHISupportsPipelineFileCache);
+	UE_LOGF(LogVulkanRHI, Log, "Vulkan PSO Precaching = %d, PipelineFileCache = %d", GRHISupportsPSOPrecaching, GRHISupportsPipelineFileCache);
 
 	// Vulkan always uses SPIR-V. This is used to differentiate standard and preview shaders in the DDC.
 	GRHIGlobals.PreferredPreviewShaderCodeFormat = TEXT("SPIRV");
@@ -453,7 +453,7 @@ FVulkanDynamicRHI::FVulkanDynamicRHI()
 	// Setup the validation requests ready before we load dlls
 	SetupValidationRequests();
 
-	UE_LOG(LogVulkanRHI, Display, TEXT("Built with Vulkan header version %u.%u.%u"), VK_API_VERSION_MAJOR(VK_HEADER_VERSION_COMPLETE), VK_API_VERSION_MINOR(VK_HEADER_VERSION_COMPLETE), VK_API_VERSION_PATCH(VK_HEADER_VERSION_COMPLETE));
+	UE_LOGF(LogVulkanRHI, Display, "Built with Vulkan header version %u.%u.%u", VK_API_VERSION_MAJOR(VK_HEADER_VERSION_COMPLETE), VK_API_VERSION_MINOR(VK_HEADER_VERSION_COMPLETE), VK_API_VERSION_PATCH(VK_HEADER_VERSION_COMPLETE));
 
 	CreateInstance();
 	SelectDevice();
@@ -479,7 +479,7 @@ void FVulkanDynamicRHI::Init()
 		// Truncate GTexturePoolSize to MB (but still counted in bytes)
 		GTexturePoolSize = int64(FGenericPlatformMath::TruncToFloat(PoolSize / 1024.0f / 1024.0f)) * 1024 * 1024;
 
-		UE_LOG(LogRHI, Log, TEXT("Texture pool is %llu MB (%d%% of %llu MB)"),
+		UE_LOGF(LogRHI, Log, "Texture pool is %llu MB (%d%% of %llu MB)",
 			GTexturePoolSize / 1024 / 1024,
 			GPoolSizeVRAMPercentage,
 			TotalGPUMemory / 1024 / 1024);
@@ -489,7 +489,7 @@ void FVulkanDynamicRHI::Init()
 		GTexturePoolSize = (int64)StreamingPoolSizeValue * 1024 * 1024;
 
 		const uint64 TotalGPUMemory = Device->GetDeviceMemoryManager().GetTotalMemory(true);
-		UE_LOG(LogRHI,Log,TEXT("Texture pool is %llu MB (of %llu MB total graphics mem)"),
+		UE_LOGF(LogRHI, Log, "Texture pool is %llu MB (of %llu MB total graphics mem)",
 				GTexturePoolSize / 1024 / 1024,
 				TotalGPUMemory / 1024 / 1024);
 	}
@@ -609,16 +609,16 @@ void FVulkanDynamicRHI::CreateInstance()
 
 			if (!bRayTracingAllowedOnCurrentShaderPlatform)
 			{
-				UE_LOG(LogVulkanRHI, Display, TEXT("Vulkan RayTracing disabled because SM6 shader platform is required."));
+				UE_LOGF(LogVulkanRHI, Display, "Vulkan RayTracing disabled because SM6 shader platform is required.");
 			}
 			else
 			{
-				UE_LOG(LogVulkanRHI, Display, TEXT("Vulkan RayTracing disabled because of failed profile check."));
+				UE_LOGF(LogVulkanRHI, Display, "Vulkan RayTracing disabled because of failed profile check.");
 			}
 		}
 	}
 
-	UE_LOG(LogVulkanRHI, Log, TEXT("Using API Version %u.%u."), VK_API_VERSION_MAJOR(ApiVersion), VK_API_VERSION_MINOR(ApiVersion));
+	UE_LOGF(LogVulkanRHI, Log, "Using API Version %u.%u.", VK_API_VERSION_MAJOR(ApiVersion), VK_API_VERSION_MINOR(ApiVersion));
 
 	// EngineName will be of the form "UnrealEngine4.21", with the minor version ("21" in this example)
 	// updated with every quarterly release
@@ -659,12 +659,11 @@ void FVulkanDynamicRHI::CreateInstance()
 	
 	if (Result == VK_ERROR_LAYER_NOT_PRESENT)
 	{
-		UE_LOG(LogVulkanRHI, Warning, TEXT("Vulkan instance creation returned an error with the requested layers (%d):"), InstanceLayers.Num());
+		UE_LOGF(LogVulkanRHI, Warning, "Vulkan instance creation returned an error with the requested layers (%d):", InstanceLayers.Num());
 
 		for (const ANSICHAR* AnsiLayerName : InstanceLayers)
 		{
-			const FString LayerStr = ANSI_TO_TCHAR(AnsiLayerName);
-			UE_LOG(LogVulkanRHI, Warning, TEXT("- %s"), *LayerStr);
+			UE_LOGF(LogVulkanRHI, Warning, "- %s", AnsiLayerName);
 		}
 
 		const EAppReturnType::Type MsgBoxResult = FPlatformMisc::MessageBoxExt(EAppMsgType::YesNo, TEXT(
@@ -725,9 +724,8 @@ void FVulkanDynamicRHI::CreateInstance()
 
 			if (!bExtensionFound)
 			{
-				FString ExtensionStr = ANSI_TO_TCHAR(Extension);
-				UE_LOG(LogVulkanRHI, Error, TEXT("Missing required Vulkan extension: %s"), *ExtensionStr);
-				MissingExtensions += ExtensionStr + TEXT("\n");
+				UE_LOGF(LogVulkanRHI, Error, "Missing required Vulkan extension: %s", Extension);
+				MissingExtensions += FString(Extension) + TEXT("\n");
 			}
 		}
 
@@ -771,7 +769,7 @@ void FVulkanDynamicRHI::SelectDevice()
 		return;
 	}
 
-	UE_LOG(LogVulkanRHI, Log, TEXT("Creating Vulkan Device using VkPhysicalDevice 0x%x."), reinterpret_cast<const void*>(PhysicalDevice));
+	UE_LOGF(LogVulkanRHI, Log, "Creating Vulkan Device using VkPhysicalDevice 0x%p.", reinterpret_cast<const void*>(PhysicalDevice));
 	Device = new FVulkanDevice(this, PhysicalDevice);
 
 	const VkPhysicalDeviceProperties& Props = Device->GetDeviceProperties();
@@ -789,7 +787,7 @@ void FVulkanDynamicRHI::SelectDevice()
 		case VK_VENDOR_ID_MESA:		GRHIVendorId = (uint32)EGpuVendorId::Mesa; break;
 		default:
 			// Unhandled case
-			UE_LOG(LogVulkanRHI, Warning, TEXT("Unhandled VkVendorId %d"), (int32)VendorId);
+			UE_LOGF(LogVulkanRHI, Warning, "Unhandled VkVendorId %d", (int32)VendorId);
 			bUseVendorIdAsIs = true;
 			break;
 		}
@@ -811,12 +809,12 @@ void FVulkanDynamicRHI::SelectDevice()
 		FString GLVersion;
 #endif
 		GRHIAdapterInternalDriverVersion = FString::Printf(TEXT("%d.%d.%d|%s"), VK_VERSION_MAJOR(Props.apiVersion), VK_VERSION_MINOR(Props.apiVersion), VK_VERSION_PATCH(Props.apiVersion), *GLVersion);
-		UE_LOG(LogVulkanRHI, Log, TEXT("API Version: %s"), *GRHIAdapterInternalDriverVersion);
+		UE_LOGF(LogVulkanRHI, Log, "API Version: %ls", *GRHIAdapterInternalDriverVersion);
 	}
 	else if (PLATFORM_WINDOWS)
 	{
 		GRHIDeviceId = Props.deviceID;
-		UE_LOG(LogVulkanRHI, Log, TEXT("API Version: %d.%d.%d"), VK_VERSION_MAJOR(Props.apiVersion), VK_VERSION_MINOR(Props.apiVersion), VK_VERSION_PATCH(Props.apiVersion));
+		UE_LOGF(LogVulkanRHI, Log, "API Version: %d.%d.%d", VK_VERSION_MAJOR(Props.apiVersion), VK_VERSION_MINOR(Props.apiVersion), VK_VERSION_PATCH(Props.apiVersion));
 	}
 	else if(PLATFORM_UNIX)
 	{
@@ -835,7 +833,7 @@ void FVulkanDynamicRHI::SelectDevice()
 		GRHIDeviceId = Props.deviceID;
 		GRHIAdapterInternalDriverVersion = GRHIAdapterUserDriverVersion;
 		GRHIAdapterDriverDate = TEXT("01-01-01");  // Unused on unix systems, pick a date that will fail test if compared but passes IsValid() check
-		UE_LOG(LogVulkanRHI, Log, TEXT("     API Version: %d.%d.%d"), VK_VERSION_MAJOR(Props.apiVersion), VK_VERSION_MINOR(Props.apiVersion), VK_VERSION_PATCH(Props.apiVersion));
+		UE_LOGF(LogVulkanRHI, Log, "     API Version: %d.%d.%d", VK_VERSION_MAJOR(Props.apiVersion), VK_VERSION_MINOR(Props.apiVersion), VK_VERSION_PATCH(Props.apiVersion));
 	}
 
 	GRHIPersistentThreadGroupCount = 1440; // TODO: Revisit based on vendor/adapter/perf query
@@ -920,7 +918,7 @@ void FVulkanDynamicRHI::InitInstance()
 					GRHISupportsRayTracing = false;
 					GRHISupportsRayTracingShaders = false;
 					GRHISupportsInlineRayTracing = false;
-					UE_LOG(LogVulkanRHI, Warning, TEXT("Using MESA RADV version prior to 24.3.2, ray tracing disabled."));
+					UE_LOGF(LogVulkanRHI, Warning, "Using MESA RADV version prior to 24.3.2, ray tracing disabled.");
 				}
 			}
 			else if ((Device->GetOptionalExtensionProperties().PhysicalDeviceDriverProperties.driverID == VK_DRIVER_ID_AMD_PROPRIETARY) ||
@@ -935,7 +933,7 @@ void FVulkanDynamicRHI::InitInstance()
 					CVarPathTracing->Set(0);
 				}
 				
-				UE_LOG(LogVulkanRHI, Display, TEXT("Using AMD Propriertary Driver, ray tracing disabled."));
+				UE_LOGF(LogVulkanRHI, Display, "Using AMD Propriertary Driver, ray tracing disabled.");
 			}
 		}
 
@@ -959,7 +957,7 @@ void FVulkanDynamicRHI::InitInstance()
 
 		//#todo-rco: Add newer Nvidia also
 		GSupportsEfficientAsyncCompute = Device->HasAsyncComputeQueue();
-		UE_LOG(LogVulkanRHI, Display, TEXT("Vulkan Async Compute has been %s."), GSupportsEfficientAsyncCompute ? TEXT("ENABLED") : TEXT("DISABLED"));
+		UE_LOGF(LogVulkanRHI, Display, "Vulkan Async Compute has been %s.", GSupportsEfficientAsyncCompute ? "ENABLED" : "DISABLED");
 
 		GSupportsVolumeTextureRendering = FVulkanPlatform::SupportsVolumeTextureRendering();
 
@@ -1059,13 +1057,13 @@ void FVulkanDynamicRHI::InitInstance()
 				GRHIMinimumWaveSize = GRHIMaximumWaveSize = Device->GetDeviceSubgroupProperties().subgroupSize;
 			}
 
-			UE_LOG(LogVulkanRHI, Display, TEXT("Wave Operations have been ENABLED (wave size: min=%d max=%d)."), GRHIMinimumWaveSize, GRHIMaximumWaveSize);
+			UE_LOGF(LogVulkanRHI, Display, "Wave Operations have been ENABLED (wave size: min=%d max=%d).", GRHIMinimumWaveSize, GRHIMaximumWaveSize);
 		}
 		else
 		{
 			const uint32 MissingStageFlags = (Device->GetDeviceSubgroupProperties().supportedStages & RequiredSubgroupShaderStageFlags) ^ RequiredSubgroupShaderStageFlags;
 			const uint32 MissingOperationFlags = (Device->GetDeviceSubgroupProperties().supportedOperations & RequiredSubgroupFlags) ^ RequiredSubgroupFlags;
-			UE_LOG(LogVulkanRHI, Display, TEXT("Wave Operations have been DISABLED (missing stages=0x%x operations=0x%x)."), MissingStageFlags, MissingOperationFlags);
+			UE_LOGF(LogVulkanRHI, Display, "Wave Operations have been DISABLED (missing stages=0x%x operations=0x%x).", MissingStageFlags, MissingOperationFlags);
 		}
 
 		FHardwareInfo::RegisterHardwareInfo(NAME_RHI, TEXT("Vulkan"));
@@ -1135,7 +1133,7 @@ void FVulkanDynamicRHI::InitInstance()
 
 #if PLATFORM_WINDOWS || PLATFORM_UNIX
 		GRHIDeviceIsIntegrated = (Device->GetDeviceProperties().deviceType == VK_PHYSICAL_DEVICE_TYPE_INTEGRATED_GPU);
-		UE_LOG(LogVulkanRHI, Log, TEXT("Integrated GPU (iGPU): %s"), GRHIDeviceIsIntegrated ? TEXT("true") : TEXT("false"));
+		UE_LOGF(LogVulkanRHI, Log, "Integrated GPU (iGPU): %s", GRHIDeviceIsIntegrated ? "true" : "false");
 #endif
 
 		InitializeSubmissionPipe();
@@ -1808,18 +1806,18 @@ void FVulkanDynamicRHI::SaveValidationCache()
 					FString CacheFilename = VulkanRHI::GetValidationCacheFilename();
 					if (FFileHelper::SaveArrayToFile(Data, *CacheFilename))
 					{
-						UE_LOG(LogVulkanRHI, Display, TEXT("Saved validation cache file '%s', %d bytes"), *CacheFilename, Data.Num());
+						UE_LOGF(LogVulkanRHI, Display, "Saved validation cache file '%ls', %d bytes", *CacheFilename, Data.Num());
 					}
 				}
 				else
 				{
-					UE_LOG(LogVulkanRHI, Warning, TEXT("Failed to query Vulkan validation cache data, VkResult=%d"), Result);
+					UE_LOGF(LogVulkanRHI, Warning, "Failed to query Vulkan validation cache data, VkResult=%d", Result);
 				}
 			}
 		}
 		else
 		{
-			UE_LOG(LogVulkanRHI, Warning, TEXT("Failed to query Vulkan validation cache size, VkResult=%d"), Result);
+			UE_LOGF(LogVulkanRHI, Warning, "Failed to query Vulkan validation cache size, VkResult=%d", Result);
 		}
 	}
 }
