VulkanRHI: Allow for better external initialization of extensions for plugins.  Create a base class for external extensions (FVulkanRHIExternalDeviceExtensionBase) where feature/property structures can be specified (and queried after being filled).

#rnx

#Synced-CL 51461249

diff --git a/Engine/Source/Runtime/VulkanRHI/Private/VulkanExtensions.cpp b/Engine/Source/Runtime/VulkanRHI/Private/VulkanExtensions.cpp
--- a/Engine/Source/Runtime/VulkanRHI/Private/VulkanExtensions.cpp
+++ b/Engine/Source/Runtime/VulkanRHI/Private/VulkanExtensions.cpp
@@ -178,7 +178,7 @@ extern TAutoConsoleVariable<int32> GGPUValidationCvar;
 #endif
 
 TSharedPtr<IHeadMountedDisplayVulkanExtensions, ESPMode::ThreadSafe> FVulkanDynamicRHI::HMDVulkanExtensions;
-TArray<const ANSICHAR*> FVulkanDeviceExtension::ExternalExtensions;
+TArray<TSharedRef<FVulkanRHIExternalDeviceExtensionBase>> FVulkanDeviceExtension::ExternalExtensions;
 TArray<const ANSICHAR*> FVulkanInstanceExtension::ExternalExtensions;
 
 
@@ -2721,6 +2721,66 @@ static void FlagExtensionSupport(const TArray<VkExtensionProperties>& ExtensionP
 }
 
 
+class FVulkanExternalDeviceExtensionWrapper : public FVulkanDeviceExtension
+{
+public: 
+	FVulkanExternalDeviceExtensionWrapper(FVulkanDevice* InDevice, TSharedRef<FVulkanRHIExternalDeviceExtensionBase> InExternalDeviceExtensionBase)
+		: FVulkanDeviceExtension(InDevice, *InExternalDeviceExtensionBase->ExtensionName, VULKAN_EXTENSION_ENABLED)
+		, ExternalDeviceExtensionBase(InExternalDeviceExtensionBase)
+	{
+	}
+
+	virtual void PrePhysicalDeviceProperties(VkPhysicalDeviceProperties2KHR& PhysicalDeviceProperties2) override final
+	{
+		VkBaseOutStructure* PropertiesStruct = bRequirementsPassed ? ExternalDeviceExtensionBase->GetPropertiesStruct() : nullptr;
+		if (PropertiesStruct)
+		{
+			PropertiesStruct->pNext = (VkBaseOutStructure*)PhysicalDeviceProperties2.pNext;
+			PhysicalDeviceProperties2.pNext = PropertiesStruct;
+		}
+	}
+	
+	virtual void PostPhysicalDeviceProperties(FOptionalVulkanDeviceExtensions& ExtensionFlags) override final
+	{
+		if (bRequirementsPassed)
+		{
+			bRequirementsPassed = ExternalDeviceExtensionBase->PropertiesCallback();
+		}
+	}
+
+	virtual void PrePhysicalDeviceFeatures(VkPhysicalDeviceFeatures2KHR& PhysicalDeviceFeatures2) override final
+	{
+		VkBaseOutStructure* FeaturesStruct = bRequirementsPassed ? ExternalDeviceExtensionBase->GetFeaturesStruct() : nullptr;
+		if (FeaturesStruct)
+		{
+			FeaturesStruct->pNext = (VkBaseOutStructure*)PhysicalDeviceFeatures2.pNext;
+			PhysicalDeviceFeatures2.pNext = FeaturesStruct;
+		}
+	}
+	
+	virtual void PostPhysicalDeviceFeatures(FOptionalVulkanDeviceExtensions& ExtensionFlags) override final
+	{
+		if (bRequirementsPassed)
+		{
+			bRequirementsPassed = ExternalDeviceExtensionBase->FeaturesCallback();
+		}
+	}
+
+	virtual void PreCreateDevice(VkDeviceCreateInfo& DeviceCreateInfo)  override final
+	{
+		VkBaseOutStructure* FeaturesStruct = bRequirementsPassed ? ExternalDeviceExtensionBase->GetFeaturesStruct() : nullptr;
+		if (FeaturesStruct)
+		{
+			FeaturesStruct->pNext = (VkBaseOutStructure*)DeviceCreateInfo.pNext;
+			DeviceCreateInfo.pNext = FeaturesStruct;
+		}
+	}
+
+private:
+	TSharedRef<FVulkanRHIExternalDeviceExtensionBase> ExternalDeviceExtensionBase;
+};
+
+
 
 FVulkanDeviceExtensionArray FVulkanDeviceExtension::GetUESupportedDeviceExtensions(FVulkanDevice* InDevice, uint32 ApiVersion)
 {
@@ -2828,22 +2888,24 @@ FVulkanDeviceExtensionArray FVulkanDeviceExtension::GetUESupportedDeviceExtensio
 	// Add in platform specific extensions
 	FVulkanPlatform::GetDeviceExtensions(InDevice, OutUEDeviceExtensions);
 
-	// Helper function to go through a list of extensions and activate them (or add them)
-	auto ActivateExternalExtensions = [&](TArray<const ANSICHAR*> Extensions, const ANSICHAR* Requester) {
-		for (const ANSICHAR* ExtensionName : Extensions)
+	// Merge extensions from outside the RHI into the main array using the wrapper
+	for (TSharedRef<FVulkanRHIExternalDeviceExtensionBase>& ExtensionBase : ExternalExtensions)
+	{
+		const int32 ExtensionIndex = FindExtension(OutUEDeviceExtensions, *ExtensionBase->ExtensionName);
+		if (ExtensionIndex == INDEX_NONE)
 		{
-			const int32 ExtensionIndex = FindExtension(OutUEDeviceExtensions, ExtensionName);
-			if (ExtensionIndex == INDEX_NONE)
-			{
-				OutUEDeviceExtensions.Add(MakeUnique<FVulkanDeviceExtension>(InDevice, ExtensionName, VULKAN_EXTENSION_ENABLED));
-				UE_LOGF(LogVulkanRHI, Warning, "%s requested device extension [%s] isn't part of the engine's core extension list. Adding it on-the-fly...", Requester, ExtensionName);
-			}
-			else
-			{
-				OutUEDeviceExtensions[ExtensionIndex]->SetActivated();
-			}
+			OutUEDeviceExtensions.Add(MakeUnique<FVulkanExternalDeviceExtensionWrapper>(InDevice, ExtensionBase));
+			UE_LOGF(LogVulkanRHI, Warning, "Adding externally requested extension: %s", *ExtensionBase->ExtensionName);
 		}
-	};
+		else
+		{
+			OutUEDeviceExtensions[ExtensionIndex]->SetActivated();
+			ExtensionBase->DuplicateCallback();
+		}
+	}
+
+	// Release the external extensions array, the device extensions are holding the references now
+	FVulkanDeviceExtension::ExternalExtensions.Empty();
 
 	// Add HMD requested extensions
 	{
@@ -2855,13 +2917,19 @@ FVulkanDeviceExtensionArray FVulkanDeviceExtension::GetUESupportedDeviceExtensio
 		{
 			TArray<const ANSICHAR*> HMDExtensions;
 			FVulkanDynamicRHI::HMDVulkanExtensions->GetVulkanDeviceExtensionsRequired(InDevice->GetPhysicalHandle(), HMDExtensions);
-			ActivateExternalExtensions(HMDExtensions, "HMD");
+
+			for (const ANSICHAR* ExtensionName : HMDExtensions)
+			{
+				const int32 ExtensionIndex = FindExtension(OutUEDeviceExtensions, ExtensionName);
+				if (ExtensionIndex == INDEX_NONE)
+				{
+					OutUEDeviceExtensions.Add(MakeUnique<FVulkanExternalDeviceExtensionWrapper>(InDevice, MakeShared<FVulkanRHIExternalDeviceExtensionBase>(ExtensionName)));
+					UE_LOGF(LogVulkanRHI, Warning, "Adding HMD requested extension: %s", ExtensionName);
+				}
+			}
 		}
 	}
 
-	// Add extensions added outside the RHI (eg plugins)
-	ActivateExternalExtensions(ExternalExtensions, "Externally");
-
 	// Now that all the extensions are listed, update their support flags
 	FlagExtensionSupport(GetDriverSupportedDeviceExtensions(InDevice->GetPhysicalHandle()), OutUEDeviceExtensions, ApiVersion, "device");
 
@@ -3007,4 +3075,3 @@ FVulkanInstanceExtensionArray FVulkanInstanceExtension::GetUESupportedInstanceEx
 
 	return OutUEInstanceExtensions;
 }
-
diff --git a/Engine/Source/Runtime/VulkanRHI/Private/VulkanExtensions.h b/Engine/Source/Runtime/VulkanRHI/Private/VulkanExtensions.h
--- a/Engine/Source/Runtime/VulkanRHI/Private/VulkanExtensions.h
+++ b/Engine/Source/Runtime/VulkanRHI/Private/VulkanExtensions.h
@@ -127,7 +127,7 @@ public:
 	virtual void PreCreateDevice(VkDeviceCreateInfo& DeviceInfo) {}
 
 	// Holds extensions requested externally (eg plugins)
-	static TArray<const ANSICHAR*> ExternalExtensions;
+	static TArray<TSharedRef<class FVulkanRHIExternalDeviceExtensionBase>> ExternalExtensions;
 
 protected:
 	struct FOptionalVulkanDeviceExtensionProperties& GetDeviceExtensionProperties();
diff --git a/Engine/Source/Runtime/VulkanRHI/Private/VulkanLayers.cpp b/Engine/Source/Runtime/VulkanRHI/Private/VulkanLayers.cpp
--- a/Engine/Source/Runtime/VulkanRHI/Private/VulkanLayers.cpp
+++ b/Engine/Source/Runtime/VulkanRHI/Private/VulkanLayers.cpp
@@ -299,7 +299,7 @@ void IVulkanDynamicRHI::AddEnabledInstanceExtensionsAndLayers(TArrayView<const A
 	FVulkanIntanceSetupHelper::ExternalLayers.Append(InInstanceLayers.GetData(), InInstanceLayers.Num());
 }
 
-void IVulkanDynamicRHI::AddEnabledDeviceExtensionsAndLayers(TArrayView<const ANSICHAR* const> InDeviceExtensions, TArrayView<const ANSICHAR* const> InDeviceLayers)
+void IVulkanDynamicRHI::AddEnabledDeviceExtensionsAndLayers(TArrayView<TSharedRef<FVulkanRHIExternalDeviceExtensionBase>> InDeviceExtensions, TArrayView<const ANSICHAR* const> InDeviceLayers)
 {
 	checkf(!GVulkanRHI, TEXT("AddEnabledDeviceExtensionsAndLayers should be called before the VulkanRHI has been created"));
 	FVulkanDeviceExtension::ExternalExtensions.Append(InDeviceExtensions.GetData(), InDeviceExtensions.Num());
diff --git a/Engine/Source/Runtime/VulkanRHI/Public/IVulkanDynamicRHI.h b/Engine/Source/Runtime/VulkanRHI/Public/IVulkanDynamicRHI.h
--- a/Engine/Source/Runtime/VulkanRHI/Public/IVulkanDynamicRHI.h
+++ b/Engine/Source/Runtime/VulkanRHI/Public/IVulkanDynamicRHI.h
@@ -31,6 +31,54 @@ struct FVulkanRHIExternalImageDeleteCallbackInfo
 	void (*Function)(void* UserData) = nullptr;
 };
 
+
+// Base class used for external code (such as plugins) to enable extensions in the Vulkan RHI. The lifetime is managed through the use of TSharedRef.
+// Basic usage:  Using this as a base class, pass in the ExtensionName and overload GetFeaturesStruct() to return the extension specific features struct.
+// IMPORTANT:    Make sure the feature struct has its sType correctly set and pNext to NULL.
+// 
+// Example to enable VK_EXT_memory_decompression:
+// class FDeviceExtensionMemoryDecompression : public FVulkanRHIExternalDeviceExtensionBase
+// {
+//    FDeviceExtensionMemoryDecompression() : FVulkanRHIExternalDeviceExtensionBase(VK_EXT_MEMORY_DECOMPRESSION_EXTENSION_NAME)
+//    {
+//        ZeroVulkanStruct(Features, VK_STRUCTURE_TYPE_PHYSICAL_DEVICE_MEMORY_DECOMPRESSION_FEATURES_EXT);
+//    }
+//    virtual VkBaseOutStructure* GetFeaturesStruct() final override
+//    {
+//       return reinterpret_cast<VkBaseOutStructure*>(&Features);
+//    }
+//    VkPhysicalDeviceMemoryDecompressionFeaturesEXT Features;
+// }
+// 
+// - If properties are needed, overload GetPropertiesStruct() to return the extension specific properties struct (with sType and pNext initialized).
+// - FeaturesCallback() / PropertiesCallback() are called after those property structures are filled.  
+// - If inspection of the features and/or properties reveals the extension should not be loaded, the callbacks should return false.
+// - If the extension is already handled by the engine or loaded by another plugin, DuplicateCallback() will be called.
+class FVulkanRHIExternalDeviceExtensionBase : public TSharedFromThis<FVulkanRHIExternalDeviceExtensionBase>
+{
+public:
+	FVulkanRHIExternalDeviceExtensionBase(const ANSICHAR* InExtensionName) : ExtensionName(InExtensionName) {}
+	virtual ~FVulkanRHIExternalDeviceExtensionBase() {}
+
+	// Returns a pointer to the extension's feature structure
+	virtual VkBaseOutStructure* GetFeaturesStruct() { return nullptr; }
+
+	// Callback when the features structure has been filled.  Return true if the features meet your requirements.
+	virtual bool FeaturesCallback() { return true;  }
+
+	// Pointer to the extension's properties structure
+	virtual VkBaseOutStructure* GetPropertiesStruct() { return nullptr;  }
+
+	// Callback when the properties structure has been filled.  Return true if the properties meet your requirements.
+	virtual bool PropertiesCallback() { return true; }
+
+	// Callback if the extension is found to be a duplicate (already loaded by the engine or other plugin).
+	virtual void DuplicateCallback() {}
+
+	// Name of the extension to be loaded
+	const FAnsiString ExtensionName;
+};
+
 enum class EVulkanRHIRunOnQueueType
 {
 	Graphics = 0,
@@ -100,7 +148,18 @@ struct IVulkanDynamicRHI : public FDynamicRHI
 	virtual void           RHIVerifyResult(VkResult Result, const ANSICHAR* VkFuntion, const ANSICHAR* Filename, uint32 Line) = 0;
 
 	static VULKANRHI_API void AddEnabledInstanceExtensionsAndLayers(TArrayView<const ANSICHAR* const> InInstanceExtensions, TArrayView<const ANSICHAR* const> InInstanceLayers);
-	static VULKANRHI_API void AddEnabledDeviceExtensionsAndLayers(TArrayView<const ANSICHAR* const> InDeviceExtensions, TArrayView<const ANSICHAR* const> InDeviceLayers);
+	static void AddEnabledDeviceExtensionsAndLayers(TArrayView<const ANSICHAR* const> InDeviceExtensions, TArrayView<const ANSICHAR* const> InDeviceLayers)
+	{
+		TArray<TSharedRef<FVulkanRHIExternalDeviceExtensionBase>> DeviceExtensions;
+		DeviceExtensions.Reserve(InDeviceExtensions.Num());
+		for (const ANSICHAR* const ExtensionName : InDeviceExtensions)
+		{
+			DeviceExtensions.Add(MakeShared<FVulkanRHIExternalDeviceExtensionBase>(ExtensionName));
+		}
+		AddEnabledDeviceExtensionsAndLayers(DeviceExtensions, InDeviceLayers);
+	}
+	static VULKANRHI_API void AddEnabledDeviceExtensionsAndLayers(TArrayView<TSharedRef<FVulkanRHIExternalDeviceExtensionBase>> InDeviceExtensions, TArrayView<const ANSICHAR* const> InDeviceLayers);
+
 
 	// Runs code on SubmissionThread with access to VkQueue.  Useful for plugins. 
 	virtual void RHIRunOnQueue(EVulkanRHIRunOnQueueType QueueType, TFunction<void(VkQueue)>&& CodeToRun, bool bWaitForSubmission) = 0;
