VulkanRHI: Fix layouts from sync2 being used when the extension is not enabled.

#rnx

#Synced-CL 51414939

diff --git a/Engine/Source/Runtime/VulkanRHI/Private/VulkanRenderpass.h b/Engine/Source/Runtime/VulkanRHI/Private/VulkanRenderpass.h
--- a/Engine/Source/Runtime/VulkanRHI/Private/VulkanRenderpass.h
+++ b/Engine/Source/Runtime/VulkanRHI/Private/VulkanRenderpass.h
@@ -80,13 +80,23 @@ inline void FVulkanAttachmentReference<VkAttachmentReference2>::SetDepthStencilA
 	{
 		bSupportsParallelRendering = false;
 	}
-	const bool bHasStencilReference = StencilReference && (StencilReference->stencilLayout != VK_IMAGE_LAYOUT_UNDEFINED);
 
 	sType = VK_STRUCTURE_TYPE_ATTACHMENT_REFERENCE_2;
-	pNext = (bSupportsParallelRendering && bHasStencilReference) ? StencilReference : nullptr;
 	attachment = AttachmentReferenceIn.attachment;
-	layout = (bSupportsParallelRendering || !bHasStencilReference) ? AttachmentReferenceIn.layout : VulkanRHI::GetMergedDepthStencilLayout(AttachmentReferenceIn.layout, StencilReference->stencilLayout);
 	aspectMask = AspectMask;
+
+	if (bSupportsParallelRendering)
+	{
+		const bool bHasStencilReference = StencilReference && (StencilReference->stencilLayout != VK_IMAGE_LAYOUT_UNDEFINED);
+		pNext = bHasStencilReference ? StencilReference : nullptr;
+		layout = AttachmentReferenceIn.layout;
+	}
+	else
+	{
+		pNext = nullptr;
+		const VkImageLayout StencilLayout = StencilReference ? StencilReference->stencilLayout : VK_IMAGE_LAYOUT_UNDEFINED;
+		layout = VulkanRHI::GetMergedDepthStencilLayout(AttachmentReferenceIn.layout, StencilLayout);
+	}
 }
 
 
@@ -340,8 +350,18 @@ struct FVulkanAttachmentDescription<VkAttachmentDescription2>
 		storeOp = InDesc.storeOp;
 		stencilLoadOp = InDesc.stencilLoadOp;
 		stencilStoreOp = InDesc.stencilStoreOp;
-		initialLayout = (bSupportsParallelRendering || (InStencilDesc == nullptr)) ? InDesc.initialLayout : VulkanRHI::GetMergedDepthStencilLayout(InDesc.initialLayout, InStencilDesc->stencilInitialLayout);
-		finalLayout = (bSupportsParallelRendering || (InStencilDesc == nullptr)) ? InDesc.finalLayout : VulkanRHI::GetMergedDepthStencilLayout(InDesc.finalLayout, InStencilDesc->stencilFinalLayout);
+		if (bSupportsParallelRendering)
+		{
+			initialLayout = InDesc.initialLayout;
+			finalLayout = InDesc.finalLayout;
+		}
+		else
+		{
+			const VkImageLayout StencilInitialLayout = InStencilDesc ? InStencilDesc->stencilInitialLayout : VK_IMAGE_LAYOUT_UNDEFINED;
+			initialLayout = VulkanRHI::GetMergedDepthStencilLayout(InDesc.initialLayout, StencilInitialLayout);
+			const VkImageLayout StencilFinalLayout = InStencilDesc ? InStencilDesc->stencilFinalLayout : VK_IMAGE_LAYOUT_UNDEFINED;
+			finalLayout = VulkanRHI::GetMergedDepthStencilLayout(InDesc.finalLayout, StencilFinalLayout);
+		}
 	}
 };
 
