VulkanRHI: Allow validation layer in test builds for Android.

#rnx

#Synced-CL 51515035

diff --git a/Engine/Source/Runtime/VulkanRHI/Private/Android/VulkanAndroidPlatform.h b/Engine/Source/Runtime/VulkanRHI/Private/Android/VulkanAndroidPlatform.h
--- a/Engine/Source/Runtime/VulkanRHI/Private/Android/VulkanAndroidPlatform.h
+++ b/Engine/Source/Runtime/VulkanRHI/Private/Android/VulkanAndroidPlatform.h
@@ -25,6 +25,10 @@
 #define VULKAN_SUPPORTS_DESCRIPTOR_INDEXING			1
 #define VULKAN_SUPPORTS_RAY_TRACING_POSITION_FETCH	0
 
+// Enabling debugging gives access to features that have little cost unless enabled (Validation Layer, Renderdoc detection, etc.)
+// Because Development is slower on Android, more testing using validation layer occurs on Test builds.
+#define VULKAN_HAS_DEBUGGING_ENABLED				(UE_BUILD_DEBUG || UE_BUILD_DEVELOPMENT || UE_BUILD_TEST)
+
 #define UE_VK_API_VERSION							VK_API_VERSION_1_1
 
 #if UE_BUILD_DEBUG || UE_BUILD_DEVELOPMENT
diff --git a/Engine/Source/Runtime/VulkanRHI/Private/VulkanMemory.cpp b/Engine/Source/Runtime/VulkanRHI/Private/VulkanMemory.cpp
--- a/Engine/Source/Runtime/VulkanRHI/Private/VulkanMemory.cpp
+++ b/Engine/Source/Runtime/VulkanRHI/Private/VulkanMemory.cpp
@@ -4600,7 +4600,7 @@ namespace VulkanRHI
 		{
 			FScopeLock ScopeLock(&CS);
 
-#if VULKAN_HAS_DEBUGGING_ENABLED
+#if (UE_BUILD_DEBUG || UE_BUILD_DEVELOPMENT)
 			FEntry* ExistingEntry = Entries.FindByPredicate([&](const FEntry& InEntry)
 				{
 					return (InEntry.Handle == Entry.Handle) && (InEntry.StructureType == Entry.StructureType);
