VulkanRHI:  Create an interface for plugins to be able to use the RHI's allocator for resources.

#rnx

#Synced-CL 51556484

diff --git a/Engine/Source/Runtime/VulkanRHI/Private/VulkanDynamicRHI.h b/Engine/Source/Runtime/VulkanRHI/Private/VulkanDynamicRHI.h
--- a/Engine/Source/Runtime/VulkanRHI/Private/VulkanDynamicRHI.h
+++ b/Engine/Source/Runtime/VulkanRHI/Private/VulkanDynamicRHI.h
@@ -82,6 +82,7 @@ public:
 	virtual const uint8* RHIGetVulkanDeviceUUID() const final override;
 	virtual VkPhysicalDevice RHIGetVkPhysicalDevice() const final override;
 	virtual const VkAllocationCallbacks* RHIGetVkAllocationCallbacks() final override;
+	virtual IVulkanMemoryAllocator* RHIGetVulkanMemoryAllocator() final override;
 	virtual VkQueue RHIGetGraphicsVkQueue() const final override;
 	virtual uint32 RHIGetGraphicsQueueIndex() const final override;
 	virtual uint32 RHIGetGraphicsQueueFamilyIndex() const final override;
diff --git a/Engine/Source/Runtime/VulkanRHI/Private/VulkanMemory.cpp b/Engine/Source/Runtime/VulkanRHI/Private/VulkanMemory.cpp
--- a/Engine/Source/Runtime/VulkanRHI/Private/VulkanMemory.cpp
+++ b/Engine/Source/Runtime/VulkanRHI/Private/VulkanMemory.cpp
@@ -1381,11 +1381,13 @@ namespace VulkanRHI
 		{
 			check(IsMapped());
 			check(InOffset + InSize <= Size);
-			VkMappedMemoryRange Range;
-			ZeroVulkanStruct(Range, VK_STRUCTURE_TYPE_MAPPED_MEMORY_RANGE);
-			Range.memory = Handle;
-			Range.offset = InOffset;
-			Range.size = InSize;
+			const VkMappedMemoryRange Range = {
+				.sType = VK_STRUCTURE_TYPE_MAPPED_MEMORY_RANGE,
+				.pNext = nullptr,
+				.memory = Handle,
+				.offset = InOffset,
+				.size = InSize
+			};
 			VERIFYVULKANRESULT(VulkanRHI::vkFlushMappedMemoryRanges(DeviceHandle, 1, &Range));
 		}
 	}
@@ -1396,11 +1398,13 @@ namespace VulkanRHI
 		{
 			check(IsMapped());
 			check(InOffset + InSize <= Size);
-			VkMappedMemoryRange Range;
-			ZeroVulkanStruct(Range, VK_STRUCTURE_TYPE_MAPPED_MEMORY_RANGE);
-			Range.memory = Handle;
-			Range.offset = InOffset;
-			Range.size = InSize;
+			const VkMappedMemoryRange Range = {
+				.sType = VK_STRUCTURE_TYPE_MAPPED_MEMORY_RANGE,
+				.pNext = nullptr,
+				.memory = Handle,
+				.offset = InOffset,
+				.size = InSize
+			};
 			VERIFYVULKANRESULT(VulkanRHI::vkInvalidateMappedMemoryRanges(DeviceHandle, 1, &Range));
 		}
 	}
@@ -3140,7 +3144,7 @@ namespace VulkanRHI
 		return OutAllocation.IsValid();
 	}
 
-	bool FMemoryManager::AllocateDedicatedImageMemory(FVulkanAllocation& OutAllocation, FVulkanEvictable* AllocationOwner, VkImage Image, VkMemoryPropertyFlags MemoryPropertyFlags, EVulkanAllocationMetaType MetaType, bool bExternal, const char* File, uint32 Line)
+	bool FMemoryManager::AllocateDedicatedImageMemory(FVulkanAllocation& OutAllocation, FVulkanEvictable* AllocationOwner, VkImage Image, VkMemoryPropertyFlags MemoryPropertyFlags, EVulkanAllocationMetaType MetaType, bool bExternal, bool bForceDedicatedAllocation, const char* File, uint32 Line)
 	{
 		SCOPED_NAMED_EVENT(FVulkanMemoryManager_AllocateDedicatedImageMemory, FColor::Cyan);
 		VkImageMemoryRequirementsInfo2KHR ImageMemoryReqs2;
@@ -3157,7 +3161,7 @@ namespace VulkanRHI
 		VulkanRHI::vkGetImageMemoryRequirements2(Device.GetHandle(), &ImageMemoryReqs2, &MemoryReqs2);
 		const VkMemoryRequirements& MemoryReqs = MemoryReqs2.memoryRequirements;
 
-		const bool bUseDedicated = DedMemoryReqs.prefersDedicatedAllocation != VK_FALSE || DedMemoryReqs.requiresDedicatedAllocation != VK_FALSE;
+		const bool bUseDedicated = bForceDedicatedAllocation || (DedMemoryReqs.prefersDedicatedAllocation != VK_FALSE || DedMemoryReqs.requiresDedicatedAllocation != VK_FALSE);
 		if (bUseDedicated)
 		{
 			uint32 TypeIndex = 0;
@@ -5158,3 +5162,151 @@ const VkPhysicalDeviceMemoryProperties& VulkanRHI::FDeviceMemoryManager::GetMemo
 {
 	return MemoryProperties;
 }
+
+// Patch external allocations into the engine's allocators
+class FVulkanMemoryAllocator : public IVulkanMemoryAllocator
+{
+public:
+	FVulkanMemoryAllocator(FVulkanDevice& InDevice) : Device(InDevice), MemoryManager(InDevice.GetMemoryManager()) {}
+	virtual ~FVulkanMemoryAllocator() {}
+
+	virtual VkResult AllocateImageMemory(VkImage Image, VkMemoryPropertyFlags Flags, bool bForceDedicatedAllocation, AllocHandle* OutAllocHandle) override final
+	{
+		const VulkanRHI::EVulkanAllocationMetaType MetaType = VulkanRHI::EVulkanAllocationMetaImageOther;
+
+		// AllocateDedicatedImageMemory will let the driver decide if allocation should be dedicated
+		VulkanRHI::FVulkanAllocation* Allocation = new VulkanRHI::FVulkanAllocation();
+		if (Device.GetMemoryManager().AllocateDedicatedImageMemory(*Allocation, nullptr, Image, Flags, MetaType, false, bForceDedicatedAllocation, __FILE__, __LINE__))
+		{
+			(*OutAllocHandle) = Allocation;
+			return VK_SUCCESS;
+		}
+		else
+		{
+			delete Allocation;
+			(*OutAllocHandle) = nullptr;
+
+			const bool bIsHostMemory = VKHasAnyFlags(Flags, VK_MEMORY_PROPERTY_HOST_VISIBLE_BIT);
+			return bIsHostMemory ? VK_ERROR_OUT_OF_HOST_MEMORY : VK_ERROR_OUT_OF_DEVICE_MEMORY;
+		}
+	}
+
+	virtual VkResult AllocateBufferMemory(VkBuffer Buffer, VkMemoryPropertyFlags Flags, bool bForceDedicatedAllocation, AllocHandle* OutAllocHandle) override final
+	{
+		VulkanRHI::EVulkanAllocationFlags UEAllocFlags = VulkanRHI::EVulkanAllocationFlags::None;
+
+		const bool bIsHostMemory = VKHasAnyFlags(Flags, VK_MEMORY_PROPERTY_HOST_VISIBLE_BIT);
+		if (bIsHostMemory)
+		{
+			UEAllocFlags |= VulkanRHI::EVulkanAllocationFlags::HostVisible;
+		}
+		if (VKHasAnyFlags(Flags, VK_MEMORY_PROPERTY_HOST_CACHED_BIT))
+		{
+			UEAllocFlags |= VulkanRHI::EVulkanAllocationFlags::HostCached;
+		}
+		if (VKHasAnyFlags(Flags, VK_MEMORY_PROPERTY_LAZILY_ALLOCATED_BIT))
+		{
+			UEAllocFlags |= VulkanRHI::EVulkanAllocationFlags::Memoryless;
+		}
+		if (VKHasAnyFlags(Flags, VK_MEMORY_PROPERTY_PROTECTED_BIT))
+		{
+			UEAllocFlags |= VulkanRHI::EVulkanAllocationFlags::Protected;
+		}
+		if (bForceDedicatedAllocation)
+		{
+			UEAllocFlags |= VulkanRHI::EVulkanAllocationFlags::Dedicated;
+		}
+
+		VulkanRHI::FVulkanAllocation* Allocation = new VulkanRHI::FVulkanAllocation();
+		if (MemoryManager.AllocateBufferMemory(*Allocation, Buffer, UEAllocFlags, TEXT("External")))
+		{
+			(*OutAllocHandle) = Allocation;
+			return VK_SUCCESS;
+		}
+		else
+		{
+			delete Allocation;
+			(*OutAllocHandle) = nullptr;
+			return bIsHostMemory ? VK_ERROR_OUT_OF_HOST_MEMORY : VK_ERROR_OUT_OF_DEVICE_MEMORY;
+		}
+	}
+
+	virtual void FreeMemory(const AllocHandle& AllocHandle) override final
+	{
+		VulkanRHI::FVulkanAllocation* Allocation = reinterpret_cast<VulkanRHI::FVulkanAllocation*>(AllocHandle);
+		Allocation->Free(Device);
+		delete Allocation;
+	}
+
+	virtual FVulkanRHIAllocationInfo GetAllocInfo(const AllocHandle& AllocHandle) override final
+	{
+		VulkanRHI::FVulkanAllocation* Allocation = reinterpret_cast<VulkanRHI::FVulkanAllocation*>(AllocHandle);
+		FVulkanRHIAllocationInfo Info;
+		Info.Handle = Allocation->GetDeviceMemoryHandle(&Device);
+		Info.Offset = Allocation->Offset;
+		Info.Size = Allocation->Size;
+		Info.Flags = Allocation->GetMemoryPropertyFlags(&Device);
+		return Info;
+	}
+
+	virtual void* MapMemory(const AllocHandle& AllocHandle) override final
+	{
+		VulkanRHI::FVulkanAllocation* Allocation = reinterpret_cast<VulkanRHI::FVulkanAllocation*>(AllocHandle);
+		return Allocation->GetMappedPointer(&Device);
+	}
+
+	virtual void UnmapMemory(const AllocHandle& AllocHandle) override final
+	{
+		// Not needed for current allocator since all host visible blocks are always mapped
+		// Keep the hooks around in case future allocators don't behave the same way
+	}
+
+	virtual void FlushMemory(const AllocHandle& AllocHandle, VkDeviceSize Offset, VkDeviceSize Size) override final
+	{
+		VulkanRHI::FVulkanAllocation* Allocation = reinterpret_cast<VulkanRHI::FVulkanAllocation*>(AllocHandle);
+		if (Size == VK_WHOLE_SIZE)
+		{
+			check(Offset < Allocation->Size);
+			Size = Allocation->Size - Offset;
+		}
+		else
+		{
+			checkf(Offset + Size <= Allocation->Size, TEXT("Flushing memory beyond the allocation size."));
+		}
+		VulkanRHI::FVulkanSubresourceAllocator* Allocator = Allocation->GetSubresourceAllocator(&Device);
+		Offset += Allocation->Offset;
+		Allocator->Flush(Offset, Size);
+	}
+
+	virtual void InvalidateMemory(const AllocHandle& AllocHandle, VkDeviceSize Offset, VkDeviceSize Size) override final
+	{
+		VulkanRHI::FVulkanAllocation* Allocation = reinterpret_cast<VulkanRHI::FVulkanAllocation*>(AllocHandle);
+		if (Size == VK_WHOLE_SIZE)
+		{
+			check(Offset < Allocation->Size);
+			Size = Allocation->Size - Offset;
+		}
+		else
+		{
+			checkf(Offset + Size <= Allocation->Size, TEXT("Invalidating memory beyond the allocation size."));
+		}
+		VulkanRHI::FVulkanSubresourceAllocator* Allocator = Allocation->GetSubresourceAllocator(&Device);
+		Offset += Allocation->Offset;
+		Allocator->Invalidate(Offset, Size);
+	}
+
+	virtual FRHIMemoryStats GetMemoryStats() override final
+	{
+		return Device.GetDeviceMemoryManager().GetRHIMemoryStats();
+	}
+
+protected:
+
+	FVulkanDevice& Device;
+	VulkanRHI::FMemoryManager& MemoryManager;
+};
+
+IVulkanMemoryAllocator* FVulkanDynamicRHI::RHIGetVulkanMemoryAllocator()
+{
+	return new FVulkanMemoryAllocator(*Device);
+}
diff --git a/Engine/Source/Runtime/VulkanRHI/Private/VulkanMemory.h b/Engine/Source/Runtime/VulkanRHI/Private/VulkanMemory.h
--- a/Engine/Source/Runtime/VulkanRHI/Private/VulkanMemory.h
+++ b/Engine/Source/Runtime/VulkanRHI/Private/VulkanMemory.h
@@ -749,7 +749,7 @@ namespace VulkanRHI
 		// Legacy calls
 		bool AllocateBufferPooled(FVulkanAllocation& Allocation, FVulkanEvictable* AllocationOwner, uint32 Size, uint32 MinAlignment, VkBufferUsageFlags BufferUsageFlags, VkMemoryPropertyFlags MemoryPropertyFlags, EVulkanAllocationMetaType MetaType, const char* File, uint32 Line);
 		bool AllocateImageMemory(FVulkanAllocation& Allocation, FVulkanEvictable* AllocationOwner, const VkMemoryRequirements& MemoryReqs, VkMemoryPropertyFlags MemoryPropertyFlags, EVulkanAllocationMetaType MetaType, bool bExternal, const char* File, uint32 Line);
-		bool AllocateDedicatedImageMemory(FVulkanAllocation& Allocation, FVulkanEvictable* AllocationOwner, VkImage Image, VkMemoryPropertyFlags MemoryPropertyFlags, EVulkanAllocationMetaType MetaType, bool bExternal, const char* File, uint32 Line);
+		bool AllocateDedicatedImageMemory(FVulkanAllocation& Allocation, FVulkanEvictable* AllocationOwner, VkImage Image, VkMemoryPropertyFlags MemoryPropertyFlags, EVulkanAllocationMetaType MetaType, bool bExternal, bool bForceDedicatedAllocation, const char* File, uint32 Line);
 	private:
 		bool AllocateBufferMemory(FVulkanAllocation& Allocation, FVulkanEvictable* AllocationOwner, const VkMemoryRequirements& MemoryReqs, VkMemoryPropertyFlags MemoryPropertyFlags, EVulkanAllocationMetaType MetaType, bool bExternal, bool bForceSeparateAllocation, const char* File, uint32 Line);
 
diff --git a/Engine/Source/Runtime/VulkanRHI/Private/VulkanTexture.cpp b/Engine/Source/Runtime/VulkanRHI/Private/VulkanTexture.cpp
--- a/Engine/Source/Runtime/VulkanRHI/Private/VulkanTexture.cpp
+++ b/Engine/Source/Runtime/VulkanRHI/Private/VulkanTexture.cpp
@@ -1915,7 +1915,8 @@ FVulkanTexture::FVulkanTexture(FVulkanDevice& InDevice, const FRHITextureCreateD
 			VkDeviceSize SizeToBeConsideredForDedicated = 12 * 1024 * 1024;
 			if ((bRenderTarget || MemoryRequirements.size >= SizeToBeConsideredForDedicated) && !bMemoryless && GVulkanEnableDedicatedImageMemory)
 			{
-				if (!InDevice.GetMemoryManager().AllocateDedicatedImageMemory(Allocation, this, Image, MemoryFlags, MetaType, bExternal, __FILE__, __LINE__))
+				constexpr bool bForceDedicated = false;  // it should be considered for dedicated, but don't force it and let driver decide
+				if (!InDevice.GetMemoryManager().AllocateDedicatedImageMemory(Allocation, this, Image, MemoryFlags, MetaType, bExternal, bForceDedicated, __FILE__, __LINE__))
 				{
 					checkNoEntry();
 				}
diff --git a/Engine/Source/Runtime/VulkanRHI/Public/IVulkanDynamicRHI.h b/Engine/Source/Runtime/VulkanRHI/Public/IVulkanDynamicRHI.h
--- a/Engine/Source/Runtime/VulkanRHI/Public/IVulkanDynamicRHI.h
+++ b/Engine/Source/Runtime/VulkanRHI/Public/IVulkanDynamicRHI.h
@@ -85,6 +85,28 @@ enum class EVulkanRHIRunOnQueueType
 	Transfer,
 };
 
+// Patch allocations into the engine's allocators
+class IVulkanMemoryAllocator
+{
+public:
+	typedef void* AllocHandle;
+
+	virtual ~IVulkanMemoryAllocator() {}
+
+	virtual VkResult AllocateImageMemory(VkImage Image, VkMemoryPropertyFlags Flags, bool bForceDedicatedAllocation, AllocHandle* OutAllocHandle) = 0;
+	virtual VkResult AllocateBufferMemory(VkBuffer Buffer, VkMemoryPropertyFlags Flags, bool bForceDedicatedAllocation, AllocHandle* OutAllocHandle) = 0;
+	virtual void FreeMemory(const AllocHandle&) = 0;
+	virtual FVulkanRHIAllocationInfo GetAllocInfo(const AllocHandle&) = 0;
+
+	virtual void* MapMemory(const AllocHandle&) = 0;
+	virtual void UnmapMemory(const AllocHandle&) = 0;
+
+	virtual void FlushMemory(const AllocHandle&, VkDeviceSize Offset, VkDeviceSize Size) = 0;
+	virtual void InvalidateMemory(const AllocHandle&, VkDeviceSize Offset, VkDeviceSize Size) = 0;
+
+	virtual FRHIMemoryStats GetMemoryStats() = 0;
+};
+
 struct IVulkanDynamicRHI : public FDynamicRHI
 {
 	virtual ERHIInterfaceType GetInterfaceType() const override { return ERHIInterfaceType::Vulkan; }
@@ -96,6 +118,8 @@ struct IVulkanDynamicRHI : public FDynamicRHI
 	virtual VkPhysicalDevice RHIGetVkPhysicalDevice() const = 0;
 	virtual const VkAllocationCallbacks* RHIGetVkAllocationCallbacks() = 0;
 
+	virtual IVulkanMemoryAllocator* RHIGetVulkanMemoryAllocator() = 0;
+
 	virtual VkQueue          RHIGetGraphicsVkQueue() const = 0;
 	virtual uint32           RHIGetGraphicsQueueIndex() const = 0;
 	virtual uint32           RHIGetGraphicsQueueFamilyIndex() const = 0;
