VulkanRHI: Fix validation layout error for 3d textures with maintenance9 enabled.

#rnx

#Synced-CL 51585795

diff --git a/Engine/Source/Runtime/VulkanRHI/Private/VulkanTexture.cpp b/Engine/Source/Runtime/VulkanRHI/Private/VulkanTexture.cpp
--- a/Engine/Source/Runtime/VulkanRHI/Private/VulkanTexture.cpp
+++ b/Engine/Source/Runtime/VulkanRHI/Private/VulkanTexture.cpp
@@ -1323,7 +1323,13 @@ void FVulkanTexture::UploadInitialData(FRHICommandListBase& RHICmdList, VulkanRH
 			ensure(CmdBuffer->IsOutsideRenderPass());
 			VkCommandBuffer StagingCommandBuffer = CmdBuffer->GetHandle();
 
-			const VkImageSubresourceRange SubresourceRange = FVulkanPipelineBarrier::MakeSubresourceRange(FullAspectMask, 0, Desc.NumMips, 0, LayerCount);
+			// The layers refer to slices in 3D textures created with VK_IMAGE_CREATE_2D_ARRAY_COMPATIBLE_BIT if Maintenance9 is enabled 
+			const bool bUseSlices =
+				(Desc.Dimension == ETextureDimension::Texture3D) &&
+				Device->GetOptionalExtensions().HasKHRMaintenance9 &&
+				!EnumHasAnyFlags(Desc.Flags, TexCreate_ReservedResource);
+
+			const VkImageSubresourceRange SubresourceRange = FVulkanPipelineBarrier::MakeSubresourceRange(FullAspectMask, 0, Desc.NumMips, 0, bUseSlices ? VK_REMAINING_ARRAY_LAYERS : LayerCount);
 
 			{
 				FVulkanPipelineBarrier Barrier;
@@ -1732,8 +1738,6 @@ void FVulkanDynamicRHI::InternalUpdateTexture3D(FRHICommandListBase& RHICmdList,
 
 	const VkPhysicalDeviceLimits& Limits = Device->GetLimits();
 
-	VkBufferImageCopy Region;
-	FMemory::Memzero(Region);
 	VulkanRHI::FStagingBuffer* StagingBuffer = nullptr;
 	const uint32 NumBlocksX = (uint32)FMath::DivideAndRoundUp<int32>(UpdateRegion.Width, (uint32)BlockSizeX);
 	const uint32 NumBlocksY = (uint32)FMath::DivideAndRoundUp<int32>(UpdateRegion.Height, (uint32)BlockSizeY);
@@ -1765,20 +1769,31 @@ void FVulkanDynamicRHI::InternalUpdateTexture3D(FRHICommandListBase& RHICmdList,
 	uint32 TextureSizeY = FMath::Max(1u, TextureRHI->GetSizeY() >> MipIndex);
 	uint32 TextureSizeZ = FMath::Max(1u, TextureRHI->GetSizeZ() >> MipIndex);
 
-	//Region.bufferOffset = 0;
-	// Set these to zero to assume tightly packed buffer
-	//Region.bufferRowLength = 0;
-	//Region.bufferImageHeight = 0;
-	Region.imageSubresource.aspectMask = VK_IMAGE_ASPECT_COLOR_BIT;
-	Region.imageSubresource.mipLevel = MipIndex;
-	//Region.imageSubresource.baseArrayLayer = 0;
-	Region.imageSubresource.layerCount = 1;
-	Region.imageOffset.x = UpdateRegion.DestX;
-	Region.imageOffset.y = UpdateRegion.DestY;
-	Region.imageOffset.z = UpdateRegion.DestZ;
-	Region.imageExtent.width = (uint32)FMath::Min((int32)(TextureSizeX-UpdateRegion.DestX), (int32)UpdateRegion.Width);
-	Region.imageExtent.height = (uint32)FMath::Min((int32)(TextureSizeY-UpdateRegion.DestY), (int32)UpdateRegion.Height);
-	Region.imageExtent.depth = (uint32)FMath::Min((int32)(TextureSizeZ-UpdateRegion.DestZ), (int32)UpdateRegion.Depth);
+	// The layers refer to slices in 3D textures created with VK_IMAGE_CREATE_2D_ARRAY_COMPATIBLE_BIT if Maintenance9 is enabled 
+	// (Maintenance5 allows the use of VK_REMAINING_ARRAY_LAYERS for copies)
+	const bool bUseSlices = Device->GetOptionalExtensions().HasKHRMaintenance5 && Device->GetOptionalExtensions().HasKHRMaintenance9 && !EnumHasAnyFlags(Texture->GetDesc().Flags, TexCreate_ReservedResource);
+
+	const VkBufferImageCopy Region = {
+		.bufferOffset = 0,
+		.bufferRowLength = 0,   // zero assumes tightly packed buffer
+		.bufferImageHeight = 0, // zero assumes tightly packed buffer
+		.imageSubresource = {
+			.aspectMask = VK_IMAGE_ASPECT_COLOR_BIT,
+			.mipLevel = MipIndex,
+			.baseArrayLayer = 0,
+			.layerCount = bUseSlices ? VK_REMAINING_ARRAY_LAYERS : 1
+		},
+		.imageOffset = { 
+			.x = (int32)UpdateRegion.DestX, 
+			.y = (int32)UpdateRegion.DestY,
+			.z = (int32)UpdateRegion.DestZ
+		},
+		.imageExtent = {
+			.width = (uint32)FMath::Min((int32)(TextureSizeX - UpdateRegion.DestX), (int32)UpdateRegion.Width),
+			.height = (uint32)FMath::Min((int32)(TextureSizeY - UpdateRegion.DestY), (int32)UpdateRegion.Height),
+			.depth = (uint32)FMath::Min((int32)(TextureSizeZ - UpdateRegion.DestZ), (int32)UpdateRegion.Depth)
+		}
+	};
 
 	const bool bNeedsTransition = RHICmdList.NeedsExtraTransitions();
 	RHICmdList.EnqueueLambda(TEXT("FVulkanTexture::InternalLockWrite"),
@@ -1788,6 +1803,9 @@ void FVulkanDynamicRHI::InternalUpdateTexture3D(FRHICommandListBase& RHICmdList,
 		});
 }
 
+
+
+
 FVulkanTexture::FVulkanTexture(FVulkanDevice& InDevice, const FRHITextureCreateDesc& InCreateDesc, const FRHITransientHeapAllocation* InTransientHeapAllocation)
 	: FRHITexture(InCreateDesc)
 	, Device(&InDevice)
