VulkanRHI: Give sparse buffers a HandleId for when using ES31 on desktop.

#rnx

#Synced-CL 51585795

diff --git a/Engine/Source/Runtime/VulkanRHI/Private/VulkanBuffer.cpp b/Engine/Source/Runtime/VulkanRHI/Private/VulkanBuffer.cpp
--- a/Engine/Source/Runtime/VulkanRHI/Private/VulkanBuffer.cpp
+++ b/Engine/Source/Runtime/VulkanRHI/Private/VulkanBuffer.cpp
@@ -161,6 +161,11 @@ FVulkanBuffer::FVulkanBuffer(FVulkanDevice& InDevice, const FRHIBufferCreateDesc
 			CurrentBufferAlloc.Alloc.VulkanHandle = (uint64)BufferHandle;
 			CurrentBufferAlloc.DeviceAddress = GetBufferDeviceAddress(Device, BufferHandle);
 
+			if (UseVulkanDescriptorCache())
+			{
+				CurrentBufferAlloc.Alloc.HandleId = ++GVulkanBufferViewHandleIdCounter;
+			}
+
 			ReservedResourceData = MakeUnique<FVulkanReservedResourceData>();
 			ReservedResourceData->LastBlockSize = Align(GetSize() % SparseBlockSize, GRHIGlobals.ReservedResources.TileSizeInBytes);
 			ReservedResourceData->BlockCount = GetSize() / SparseBlockSize;
