VulkanRHI: remove queue family ownership transfers for buffers since Maintenance9 is a requirement for async compute.  Enable async compute by default when requirements are met.

#rnx

#Synced-CL 51774848

diff --git a/Engine/Source/Runtime/VulkanRHI/Private/VulkanBarriers.cpp b/Engine/Source/Runtime/VulkanRHI/Private/VulkanBarriers.cpp
--- a/Engine/Source/Runtime/VulkanRHI/Private/VulkanBarriers.cpp
+++ b/Engine/Source/Runtime/VulkanRHI/Private/VulkanBarriers.cpp
@@ -10,7 +10,6 @@
 // All shader stages supported by VK device - VK_PIPELINE_STAGE_VERTEX_SHADER_BIT, FRAGMENT etc
 extern uint32 GVulkanDevicePipelineStageBits;
 
-extern int32 GVulkanAllowConcurrentBuffer;
 extern int32 GVulkanAllowConcurrentImage;
 
 int32 GVulkanUseMemoryBarrierOpt = 1;
@@ -1239,13 +1238,14 @@ void ProcessTransitionLegacy(FVulkanCommandListContext& Context, TArrayView<cons
 
 // Removes stages/access that aren't supported by the compute queue
 template<typename BarrierType>
-static void MaskSupportedAsyncFlags(FVulkanDevice& Device, TArray<BarrierType>& InOutBarriers, bool bMaskSrc, bool bMaskDst)
+static void MaskSupportedAsyncFlags(FVulkanDevice& Device, TArray<BarrierType>& InOutBarriers, bool bMaskSrc, bool bMaskDst, bool bRemoveEmpty)
 {
 	const VkPipelineStageFlags SupportedComputeStageMask = Device.GetComputeQueue()->GetSupportedStageBits();
 	const VkAccessFlags SupportedComputeAccessMasks = Device.GetComputeQueue()->GetSupportedAccessFlags();
 
-	for (BarrierType& Barrier : InOutBarriers)
+	for (int32 Index = InOutBarriers.Num() - 1; Index >= 0; --Index)
 	{
+		BarrierType& Barrier = InOutBarriers[Index];
 		if (bMaskSrc)
 		{
 			Barrier.srcStageMask &= SupportedComputeStageMask;
@@ -1256,6 +1256,10 @@ static void MaskSupportedAsyncFlags(FVulkanDevice& Device, TArray<BarrierType>&
 			Barrier.dstStageMask &= SupportedComputeStageMask;
 			Barrier.dstAccessMask &= SupportedComputeAccessMasks;
 		}
+		if (bRemoveEmpty && ((Barrier.srcStageMask == 0) || (Barrier.dstStageMask == 0)))
+		{
+			InOutBarriers.RemoveAtSwap(Index);
+		}
 	}
 }
 
@@ -1298,7 +1302,7 @@ static void PatchCrossPipeTransitions(TArray<BarrierType>& Barriers, FVulkanComm
 		const bool bSrcIsGraphics = ((SrcPipelines == ERHIPipeline::Graphics) || (DstPipelines == ERHIPipeline::AsyncCompute));
 
 		// When a transition uses ERHIPipeline::All, it might include stages/access that isn't supported on all queues
-		MaskSupportedAsyncFlags(Context.Device, Barriers, !bSrcIsGraphics, bSrcIsGraphics);
+		MaskSupportedAsyncFlags(Context.Device, Barriers, !bSrcIsGraphics, bSrcIsGraphics, false);
 
 		for (int32 Index = 0; Index < Barriers.Num(); ++Index)
 		{
@@ -1412,10 +1416,9 @@ void ProcessTransitionSync2(FVulkanCommandListContext& Context, TArrayView<const
 		};
 
 
-	TArray<VkBufferMemoryBarrier2> TempBufferBarriers;
 	TArray<VkImageMemoryBarrier2> TempImageBarriers;
 
-	const bool bUseOwnershipTransfers = Context.Device.HasAsyncComputeQueue() && (!GVulkanAllowConcurrentBuffer || !GVulkanAllowConcurrentImage);
+	const bool bUseOwnershipTransfers = Context.Device.HasAsyncComputeQueue() && !GVulkanAllowConcurrentImage;
 	for (const FRHITransition* Transition : Transitions)
 	{
 		const FVulkanTransitionData* Data = Transition->GetPrivateData<FVulkanTransitionData>();
@@ -1465,9 +1468,10 @@ void ProcessTransitionSync2(FVulkanCommandListContext& Context, TArrayView<const
 			// For cross-queue 1..1 transitions we will keep the layout change in the same barrier as the queue transfer
 			if (bUseOwnershipTransfers && (Data->SrcPipelines != Data->DstPipelines))
 			{
-				if (!GVulkanAllowConcurrentBuffer)
+				// For buffer transitions, just remove unsupported flags for the compute queue
+				if (ExecutingPipeline == ERHIPipeline::AsyncCompute)
 				{
-					PatchCrossPipeTransitions(BarrierArrays.BufferBarriers, Context, Data->SrcPipelines, Data->DstPipelines, bIsBeginTransition);
+					MaskSupportedAsyncFlags(Context.Device, BarrierArrays.BufferBarriers, true, true, true);
 				}
 
 				if (!GVulkanAllowConcurrentImage)
@@ -1489,22 +1493,14 @@ void ProcessTransitionSync2(FVulkanCommandListContext& Context, TArrayView<const
 			(IsSingleRHIPipeline(Data->SrcPipelines) && !IsSingleRHIPipeline(Data->DstPipelines) && (Data->SrcPipelines != ExecutingPipeline))));
 
 
-		// For resources without concurrent sharing mode, make copies of the array so we can:
+		// For images without concurrent sharing mode, make copies of the array so we can:
 		// - wipe the layout transitions from the barriers
 		// - patch in the actual queue family ownership transfer
-		if (bUseOwnershipTransfers && (bNeedsRelease || bNeedsAcquire))
+		// NOTE: Maintenance9 being a requirement, ownership transfers aren't required for buffers
+		if (bUseOwnershipTransfers && (bNeedsRelease || bNeedsAcquire) && !GVulkanAllowConcurrentImage)
 		{
-			if (!GVulkanAllowConcurrentBuffer)
-			{
-				TempBufferBarriers = BarrierArrays.BufferBarriers;
-				PatchCrossPipeTransitions(TempBufferBarriers, Context, Data->SrcPipelines, Data->DstPipelines, bIsBeginTransition);
-			}
-
-			if (!GVulkanAllowConcurrentImage)
-			{
-				TempImageBarriers = BarrierArrays.ImageBarriers;
-				PatchCrossPipeTransitions(TempImageBarriers, Context, Data->SrcPipelines, Data->DstPipelines, bIsBeginTransition);
-			}
+			TempImageBarriers = BarrierArrays.ImageBarriers;
+			PatchCrossPipeTransitions(TempImageBarriers, Context, Data->SrcPipelines, Data->DstPipelines, bIsBeginTransition);
 		}
 
 		if (bNeedsPreLayoutChange)
@@ -1512,17 +1508,16 @@ void ProcessTransitionSync2(FVulkanCommandListContext& Context, TArrayView<const
 			// Remove unsupported flags if we're submitting on the compute queue
 			if (ExecutingPipeline == ERHIPipeline::AsyncCompute)
 			{
-				MaskSupportedAsyncFlags(Context.Device, BarrierArrays.BufferBarriers, false, true);
-				MaskSupportedAsyncFlags(Context.Device, BarrierArrays.ImageBarriers, false, true);
+				MaskSupportedAsyncFlags(Context.Device, BarrierArrays.BufferBarriers, false, true, true);
+				MaskSupportedAsyncFlags(Context.Device, BarrierArrays.ImageBarriers, false, true, false);
 			}
 
 			SubmitBarriers(MakeArrayView(BarrierArrays.MemoryBarriers), MakeArrayView(BarrierArrays.BufferBarriers), MakeArrayView(BarrierArrays.ImageBarriers));
 		}
 
-		if (TempBufferBarriers.Num() || TempImageBarriers.Num())
+		if (TempImageBarriers.Num())
 		{
-			SubmitBarriers(MakeArrayView(BarrierArrays.MemoryBarriers), MakeArrayView(TempBufferBarriers), MakeArrayView(TempImageBarriers));
-			TempBufferBarriers.Reset();
+			SubmitBarriers(MakeArrayView(BarrierArrays.MemoryBarriers), {/* no ownership transfers for buffers */}, MakeArrayView(TempImageBarriers));
 			TempImageBarriers.Reset();
 		}
 
@@ -1531,8 +1526,8 @@ void ProcessTransitionSync2(FVulkanCommandListContext& Context, TArrayView<const
 			// Remove unsupported flags if we're submitting on the compute queue
 			if (ExecutingPipeline == ERHIPipeline::AsyncCompute)
 			{
-				MaskSupportedAsyncFlags(Context.Device, BarrierArrays.BufferBarriers, true, false);
-				MaskSupportedAsyncFlags(Context.Device, BarrierArrays.ImageBarriers, true, false);
+				MaskSupportedAsyncFlags(Context.Device, BarrierArrays.BufferBarriers, true, false, true);
+				MaskSupportedAsyncFlags(Context.Device, BarrierArrays.ImageBarriers, true, false, false);
 			}
 
 			SubmitBarriers(MakeArrayView(BarrierArrays.MemoryBarriers), MakeArrayView(BarrierArrays.BufferBarriers), MakeArrayView(BarrierArrays.ImageBarriers));
diff --git a/Engine/Source/Runtime/VulkanRHI/Private/VulkanDevice.cpp b/Engine/Source/Runtime/VulkanRHI/Private/VulkanDevice.cpp
--- a/Engine/Source/Runtime/VulkanRHI/Private/VulkanDevice.cpp
+++ b/Engine/Source/Runtime/VulkanRHI/Private/VulkanDevice.cpp
@@ -25,27 +25,17 @@
 
 static TAutoConsoleVariable<int32> GRHIAllowAsyncComputeCvar(
 	TEXT("r.Vulkan.AllowAsyncCompute"),
-	0,
-	TEXT("0 to disable async compute queue (if available)\n")
-	TEXT("1 to allow async compute queue"),
+	1,
+	TEXT("0 to disable async compute queue\n")
+	TEXT("1 to allow async compute queue if available (default)"),
 	ECVF_ReadOnly
 );
 
 static TAutoConsoleVariable<int32> GRHIAllowTransferQueueCvar(
 	TEXT("r.Vulkan.AllowTransferQueue"),
 	0,
-	TEXT("0 to disable transfer queue (if available)\n")
-	TEXT("1 to allow transfer queue"),
-	ECVF_ReadOnly
-);
-
-int32 GVulkanAllowConcurrentBuffer = 1;
-static TAutoConsoleVariable<int32> GCVarAllowConcurrentBuffer(
-	TEXT("r.Vulkan.AllowConcurrentBuffer"),
-	GVulkanAllowConcurrentBuffer,
-	TEXT("When async compute is supported: \n")
-	TEXT(" 0 to use queue family ownership transfers with buffers\n")
-	TEXT(" 1 to use sharing mode concurrent with buffers"),
+	TEXT("0 to disable transfer queue (default)\n")
+	TEXT("1 to allow transfer queue if available"),
 	ECVF_ReadOnly
 );
 
@@ -374,6 +364,7 @@ void FVulkanDevice::CreateDevice(TArray<const ANSICHAR*>& DeviceLayers, FVulkanD
 			// - maintenance9 is available
 			if ((ComputeQueueFamilyIndex == -1) &&
 				(GRHIAllowAsyncComputeCvar.GetValueOnAnyThread() != 0 || GAllowPresentOnComputeQueue.GetValueOnAnyThread() != 0 || FParse::Param(FCommandLine::Get(), TEXT("ForceAsyncCompute"))) &&
+				!FParse::Param(FCommandLine::Get(), TEXT("DisableAsyncCompute")) &&
 				(GfxQueueFamilyIndex != FamilyIndex) && 
 				SupportsParallelRendering() &&
 				OptionalDeviceExtensions.HasKHRMaintenance9)
@@ -1714,32 +1705,24 @@ void FVulkanDevice::InitGlobalSamplers()
 
 VkBuffer FVulkanDevice::CreateBuffer(VkDeviceSize BufferSize, VkBufferUsageFlags BufferUsageFlags, VkBufferCreateFlags BufferCreateFlags) const
 {
-	VkBufferCreateInfo BufferCreateInfo;
-	BufferCreateInfo.sType = VK_STRUCTURE_TYPE_BUFFER_CREATE_INFO;
-	BufferCreateInfo.pNext = nullptr;
-	BufferCreateInfo.size = BufferSize;
-	BufferCreateInfo.usage = BufferUsageFlags;
-	BufferCreateInfo.flags = BufferCreateFlags;
-
-	// For descriptors buffers
+	VkBufferCreateInfo BufferCreateInfo = 
+	{
+		.sType = VK_STRUCTURE_TYPE_BUFFER_CREATE_INFO,
+		.pNext = nullptr,
+		.flags = BufferCreateFlags,
+		.size = BufferSize,
+		.usage = BufferUsageFlags,
+		.sharingMode = VK_SHARING_MODE_EXCLUSIVE,
+		.queueFamilyIndexCount = 0,
+		.pQueueFamilyIndices = nullptr
+	};
+
+	// For bindless
 	if (GetOptionalExtensions().HasBufferDeviceAddress)
 	{
 		BufferCreateInfo.usage |= VK_BUFFER_USAGE_SHADER_DEVICE_ADDRESS_BIT;
 	}
 
-	if (HasMultipleQueues() && GVulkanAllowConcurrentBuffer)
-	{
-		BufferCreateInfo.sharingMode = VK_SHARING_MODE_CONCURRENT;
-		BufferCreateInfo.queueFamilyIndexCount = ActiveQueueFamilies.Num();
-		BufferCreateInfo.pQueueFamilyIndices = (uint32_t*)ActiveQueueFamilies.GetData();
-	}
-	else
-	{
-		BufferCreateInfo.sharingMode = VK_SHARING_MODE_EXCLUSIVE;
-		BufferCreateInfo.queueFamilyIndexCount = 0;
-		BufferCreateInfo.pQueueFamilyIndices = nullptr;
-	}
-
 	VkBuffer BufferHandle = VK_NULL_HANDLE;
 	VERIFYVULKANRESULT(VulkanRHI::vkCreateBuffer(Device, &BufferCreateInfo, VULKAN_CPU_ALLOCATOR, &BufferHandle));
 	return BufferHandle;
diff --git a/Engine/Source/Runtime/VulkanRHI/Private/VulkanTexture.cpp b/Engine/Source/Runtime/VulkanRHI/Private/VulkanTexture.cpp
--- a/Engine/Source/Runtime/VulkanRHI/Private/VulkanTexture.cpp
+++ b/Engine/Source/Runtime/VulkanRHI/Private/VulkanTexture.cpp
@@ -30,12 +30,12 @@ static FAutoConsoleVariableRef CVarVulkanDepthStencilForceStorageBit(
 	ECVF_Default
 );
 
-int32 GVulkanAllowConcurrentImage = 1;
+int32 GVulkanAllowConcurrentImage = 0;
 static TAutoConsoleVariable<int32> GCVarAllowConcurrentImage(
 	TEXT("r.Vulkan.AllowConcurrentImage"),
 	GVulkanAllowConcurrentImage,
 	TEXT("When async compute is supported: \n")
-	TEXT(" 0 to use queue family ownership transfers with images\n")
+	TEXT(" 0 to use queue family ownership transfers with images (default)\n")
 	TEXT(" 1 to use sharing mode concurrent with images"),
 	ECVF_ReadOnly
 );
