VulkanRHI: Move SupportsDeviceLocalHostVisible to the device and include more devices.

#rnx

#Synced-CL 51908825

diff --git a/Engine/Source/Runtime/VulkanRHI/Private/VulkanBindlessDescriptorManager.cpp b/Engine/Source/Runtime/VulkanRHI/Private/VulkanBindlessDescriptorManager.cpp
--- a/Engine/Source/Runtime/VulkanRHI/Private/VulkanBindlessDescriptorManager.cpp
+++ b/Engine/Source/Runtime/VulkanRHI/Private/VulkanBindlessDescriptorManager.cpp
@@ -155,8 +155,7 @@ const VkPushConstantRange FVulkanBindlessDescriptorManager::PushConstantRanges[(
 
 static inline VkMemoryPropertyFlags GetDescriptorBufferMemoryType(FVulkanDevice& Device)
 {
-	if (Device.HasUnifiedMemory() || (FVulkanPlatform::SupportsDeviceLocalHostVisibleWithNoPenalty(Device.GetVendorId()) &&
-		Device.GetDeviceMemoryManager().SupportsMemoryType(VK_MEMORY_PROPERTY_HOST_VISIBLE_BIT | VK_MEMORY_PROPERTY_HOST_COHERENT_BIT | VK_MEMORY_PROPERTY_DEVICE_LOCAL_BIT)))
+	if (Device.SupportsDeviceLocalHostVisible())
 	{
 		return VK_MEMORY_PROPERTY_DEVICE_LOCAL_BIT | VK_MEMORY_PROPERTY_HOST_VISIBLE_BIT | VK_MEMORY_PROPERTY_HOST_COHERENT_BIT;
 	}
diff --git a/Engine/Source/Runtime/VulkanRHI/Private/VulkanDevice.cpp b/Engine/Source/Runtime/VulkanRHI/Private/VulkanDevice.cpp
--- a/Engine/Source/Runtime/VulkanRHI/Private/VulkanDevice.cpp
+++ b/Engine/Source/Runtime/VulkanRHI/Private/VulkanDevice.cpp
@@ -517,6 +517,20 @@ bool FVulkanDevice::UseMinimalSubmits() const
 		GCVarVulkanUseMinimalSubmits.GetValueOnRenderThread();
 }
 
+bool FVulkanDevice::SupportsDeviceLocalHostVisible() const
+{
+	if (HasUnifiedMemory())
+	{
+		return true;
+	}
+
+	if (DeviceMemoryManager.SupportsMemoryType(VK_MEMORY_PROPERTY_HOST_VISIBLE_BIT | VK_MEMORY_PROPERTY_HOST_COHERENT_BIT | VK_MEMORY_PROPERTY_DEVICE_LOCAL_BIT))
+	{
+		return (VendorId == EGpuVendorId::Amd) || (GpuProps.deviceType == VK_PHYSICAL_DEVICE_TYPE_INTEGRATED_GPU);
+	}
+
+	return false;
+}
 
 #if VULKAN_ENABLE_DRAW_MARKERS
 TAutoConsoleVariable<int32> CVarVulkanDebugMarkers(
diff --git a/Engine/Source/Runtime/VulkanRHI/Private/VulkanDevice.h b/Engine/Source/Runtime/VulkanRHI/Private/VulkanDevice.h
--- a/Engine/Source/Runtime/VulkanRHI/Private/VulkanDevice.h
+++ b/Engine/Source/Runtime/VulkanRHI/Private/VulkanDevice.h
@@ -637,6 +637,8 @@ public:
 
 	bool UseMinimalSubmits() const;
 
+	bool SupportsDeviceLocalHostVisible() const;
+
 private:
 	void MapBufferFormatSupport(FPixelFormatInfo& PixelFormatInfo, EPixelFormat UEFormat, VkFormat VulkanFormat);
 	void MapImageFormatSupport(FPixelFormatInfo& PixelFormatInfo, const TArrayView<const VkFormat>& PrioritizedFormats, EPixelFormatCapabilities RequiredCapabilities);
diff --git a/Engine/Source/Runtime/VulkanRHI/Private/VulkanGenericPlatform.h b/Engine/Source/Runtime/VulkanRHI/Private/VulkanGenericPlatform.h
--- a/Engine/Source/Runtime/VulkanRHI/Private/VulkanGenericPlatform.h
+++ b/Engine/Source/Runtime/VulkanRHI/Private/VulkanGenericPlatform.h
@@ -89,8 +89,6 @@ public:
 	// Some platforms have issues with the access flags for the Present layout
 	static bool RequiresPresentLayoutFix() { return false; }
 
-	static bool SupportsDeviceLocalHostVisibleWithNoPenalty(EGpuVendorId VendorId) { return false; }
-
 	static bool HasUnifiedMemory() { return false; }
 
 	// Allow the platform code to restrict the device features
diff --git a/Engine/Source/Runtime/VulkanRHI/Private/Windows/VulkanWindowsPlatform.cpp b/Engine/Source/Runtime/VulkanRHI/Private/Windows/VulkanWindowsPlatform.cpp
--- a/Engine/Source/Runtime/VulkanRHI/Private/Windows/VulkanWindowsPlatform.cpp
+++ b/Engine/Source/Runtime/VulkanRHI/Private/Windows/VulkanWindowsPlatform.cpp
@@ -183,7 +183,3 @@ void FVulkanWindowsPlatform::CreateSurface(FVulkanPlatformWindowContext& WindowC
 	VERIFYVULKANRESULT(VulkanDynamicAPI::vkCreateWin32SurfaceKHR(Instance, &SurfaceCreateInfo, VULKAN_CPU_ALLOCATOR, OutSurface));
 }
 
-bool FVulkanWindowsPlatform::SupportsDeviceLocalHostVisibleWithNoPenalty(EGpuVendorId VendorId)
-{
-	return (VendorId == EGpuVendorId::Amd);
-}
diff --git a/Engine/Source/Runtime/VulkanRHI/Private/Windows/VulkanWindowsPlatform.h b/Engine/Source/Runtime/VulkanRHI/Private/Windows/VulkanWindowsPlatform.h
--- a/Engine/Source/Runtime/VulkanRHI/Private/Windows/VulkanWindowsPlatform.h
+++ b/Engine/Source/Runtime/VulkanRHI/Private/Windows/VulkanWindowsPlatform.h
@@ -55,8 +55,6 @@ public:
 
 	static void CreateSurface(FVulkanPlatformWindowContext& WindowContext, VkInstance Instance, VkSurfaceKHR* OutSurface);
 
-	static bool SupportsDeviceLocalHostVisibleWithNoPenalty(EGpuVendorId VendorId);
-
 	static VkTimeDomainKHR GetTimeDomain() { return VK_TIME_DOMAIN_QUERY_PERFORMANCE_COUNTER_KHR; }
 
 private:
