VulkanRHI: Disable async compute until issue with parallel render pass is fixed.

#rnx
#Synced-CL 51908825

diff --git a/Engine/Source/Runtime/VulkanRHI/Private/VulkanDevice.cpp b/Engine/Source/Runtime/VulkanRHI/Private/VulkanDevice.cpp
--- a/Engine/Source/Runtime/VulkanRHI/Private/VulkanDevice.cpp
+++ b/Engine/Source/Runtime/VulkanRHI/Private/VulkanDevice.cpp
@@ -25,9 +25,9 @@
 
 static TAutoConsoleVariable<int32> GRHIAllowAsyncComputeCvar(
 	TEXT("r.Vulkan.AllowAsyncCompute"),
-	1,
-	TEXT("0 to disable async compute queue\n")
-	TEXT("1 to allow async compute queue if available (default)"),
+	0,
+	TEXT("0 to disable async compute queue (default)\n")
+	TEXT("1 to allow async compute queue if available"),
 	ECVF_ReadOnly
 );
 
