VulkanRHI: start using 64bit pipeline creation flags for bindless pipelines (preparing for descriptor_heap).

#rnx

#Synced-CL 51908825

diff --git a/Engine/Source/Runtime/VulkanRHI/Private/VulkanBindlessDescriptorManager.h b/Engine/Source/Runtime/VulkanRHI/Private/VulkanBindlessDescriptorManager.h
--- a/Engine/Source/Runtime/VulkanRHI/Private/VulkanBindlessDescriptorManager.h
+++ b/Engine/Source/Runtime/VulkanRHI/Private/VulkanBindlessDescriptorManager.h
@@ -64,6 +64,11 @@ public:
 		return !IsBindlessDisabled(Configuration);
 	}
 
+	VkPipelineCreateFlagBits2 GetPipelineCreateFlag() const
+	{
+		return VK_PIPELINE_CREATE_2_DESCRIPTOR_BUFFER_BIT_EXT;
+	}
+
 	VkPipelineLayout GetPipelineLayout() const
 	{
 		return BindlessPipelineLayout;
@@ -123,7 +128,6 @@ private:
 
 	VkDescriptorBufferBindingInfoEXT BufferBindingInfo[VulkanBindless::NumBindlessSets];
 	uint32_t BufferIndices[VulkanBindless::MaxNumSets];
-
 	VkPipelineLayout BindlessPipelineLayout = VK_NULL_HANDLE;
 	TStaticArray<VkDescriptorSetLayout, VulkanBindless::MaxNumSets> BindlessDescriptorSetLayouts;
 	static const VkPushConstantRange PushConstantRanges[(uint8)EPushConstantRangeType::Count];
diff --git a/Engine/Source/Runtime/VulkanRHI/Private/VulkanPipeline.cpp b/Engine/Source/Runtime/VulkanRHI/Private/VulkanPipeline.cpp
--- a/Engine/Source/Runtime/VulkanRHI/Private/VulkanPipeline.cpp
+++ b/Engine/Source/Runtime/VulkanRHI/Private/VulkanPipeline.cpp
@@ -1413,12 +1413,19 @@ void FVulkanPipelineStateCacheManager::CreateVertexInputLibraries()
 	DynamicStateCreateInfo.dynamicStateCount = UE_ARRAY_COUNT(VertexInputDynamicStates);
 	DynamicStateCreateInfo.pDynamicStates = VertexInputDynamicStates;
 
+	const VkPipelineCreateFlags2CreateInfo Flags2CreateInfo = {
+		.sType = VK_STRUCTURE_TYPE_PIPELINE_CREATE_FLAGS_2_CREATE_INFO,
+		.pNext = &LibraryInfo,
+		.flags = VK_PIPELINE_CREATE_2_LIBRARY_BIT_KHR |
+				 Device->GetBindlessDescriptorManager()->GetPipelineCreateFlag() |
+				 (GVulkanGraphicPipelineLibraryLinkingMode == 1 ? VK_PIPELINE_CREATE_2_RETAIN_LINK_TIME_OPTIMIZATION_INFO_BIT_EXT : 0ull)
+	};
+
 	VkGraphicsPipelineCreateInfo GraphicsPipelineCreateInfo;
 	ZeroVulkanStruct(GraphicsPipelineCreateInfo, VK_STRUCTURE_TYPE_GRAPHICS_PIPELINE_CREATE_INFO);
 	GraphicsPipelineCreateInfo.sType = VK_STRUCTURE_TYPE_GRAPHICS_PIPELINE_CREATE_INFO;
-	GraphicsPipelineCreateInfo.pNext = &LibraryInfo;
-	GraphicsPipelineCreateInfo.flags = VK_PIPELINE_CREATE_LIBRARY_BIT_KHR | VK_PIPELINE_CREATE_DESCRIPTOR_BUFFER_BIT_EXT | 
-		(GVulkanGraphicPipelineLibraryLinkingMode == 1 ? VK_PIPELINE_CREATE_RETAIN_LINK_TIME_OPTIMIZATION_INFO_BIT_EXT : 0);
+	GraphicsPipelineCreateInfo.pNext = &Flags2CreateInfo;
+	GraphicsPipelineCreateInfo.flags = 0; // using VkPipelineCreateFlags2CreateInfo
 	GraphicsPipelineCreateInfo.layout = Device->GetBindlessDescriptorManager()->GetPipelineLayout();
 	GraphicsPipelineCreateInfo.pVertexInputState = &VertexInputStateCreateInfo;
 	GraphicsPipelineCreateInfo.pInputAssemblyState = &InputAssemblyStateCreateInfo;
@@ -1477,12 +1484,19 @@ void FVulkanPipelineStateCacheManager::CreateEmptyFragmentStageLibrary()
 	VkPipelineDepthStencilStateCreateInfo DepthStencilStateCreateInfo;
 	ZeroVulkanStruct(DepthStencilStateCreateInfo, VK_STRUCTURE_TYPE_PIPELINE_DEPTH_STENCIL_STATE_CREATE_INFO);
 
+	const VkPipelineCreateFlags2CreateInfo Flags2CreateInfo = {
+		.sType = VK_STRUCTURE_TYPE_PIPELINE_CREATE_FLAGS_2_CREATE_INFO,
+		.pNext = &LibraryInfo,
+		.flags = VK_PIPELINE_CREATE_2_LIBRARY_BIT_KHR | 
+				 Device->GetBindlessDescriptorManager()->GetPipelineCreateFlag() |
+				 (GVulkanGraphicPipelineLibraryLinkingMode == 1 ? VK_PIPELINE_CREATE_2_RETAIN_LINK_TIME_OPTIMIZATION_INFO_BIT_EXT : 0ull)
+	};
+
 	VkGraphicsPipelineCreateInfo GraphicsPipelineCreateInfo;
 	ZeroVulkanStruct(GraphicsPipelineCreateInfo, VK_STRUCTURE_TYPE_GRAPHICS_PIPELINE_CREATE_INFO);
 	GraphicsPipelineCreateInfo.sType = VK_STRUCTURE_TYPE_GRAPHICS_PIPELINE_CREATE_INFO;
-	GraphicsPipelineCreateInfo.pNext = &LibraryInfo;
-	GraphicsPipelineCreateInfo.flags = VK_PIPELINE_CREATE_LIBRARY_BIT_KHR | VK_PIPELINE_CREATE_DESCRIPTOR_BUFFER_BIT_EXT | 
-		(GVulkanGraphicPipelineLibraryLinkingMode == 1 ? VK_PIPELINE_CREATE_RETAIN_LINK_TIME_OPTIMIZATION_INFO_BIT_EXT : 0);
+	GraphicsPipelineCreateInfo.pNext = &Flags2CreateInfo;
+	GraphicsPipelineCreateInfo.flags = 0; // using VkPipelineCreateFlags2CreateInfo
 	GraphicsPipelineCreateInfo.pMultisampleState = &MultisampleStateCreateInfo;
 	GraphicsPipelineCreateInfo.pDepthStencilState = &DepthStencilStateCreateInfo;
 	GraphicsPipelineCreateInfo.layout = Device->GetBindlessDescriptorManager()->GetPipelineLayout();
@@ -1600,12 +1614,19 @@ VkPipeline FVulkanPipelineStateCacheManager::FindOrCreateFragmentOutputStateLibr
 	DynamicStateCreateInfo.dynamicStateCount = UE_ARRAY_COUNT(FragmentOutputStateDynamicStates);
 	DynamicStateCreateInfo.pDynamicStates = FragmentOutputStateDynamicStates;
 
+	const VkPipelineCreateFlags2CreateInfo Flags2CreateInfo = {
+	.sType = VK_STRUCTURE_TYPE_PIPELINE_CREATE_FLAGS_2_CREATE_INFO,
+	.pNext = &LibraryInfo,
+	.flags = VK_PIPELINE_CREATE_2_LIBRARY_BIT_KHR | 
+			 Device->GetBindlessDescriptorManager()->GetPipelineCreateFlag() |
+			 (GVulkanGraphicPipelineLibraryLinkingMode == 1 ? VK_PIPELINE_CREATE_2_RETAIN_LINK_TIME_OPTIMIZATION_INFO_BIT_EXT : 0ull)
+	};
+
 	VkGraphicsPipelineCreateInfo GraphicsPipelineCreateInfo;
 	ZeroVulkanStruct(GraphicsPipelineCreateInfo, VK_STRUCTURE_TYPE_GRAPHICS_PIPELINE_CREATE_INFO);
 	GraphicsPipelineCreateInfo.sType = VK_STRUCTURE_TYPE_GRAPHICS_PIPELINE_CREATE_INFO;
-	GraphicsPipelineCreateInfo.pNext = &LibraryInfo;
-	GraphicsPipelineCreateInfo.flags = VK_PIPELINE_CREATE_LIBRARY_BIT_KHR | VK_PIPELINE_CREATE_DESCRIPTOR_BUFFER_BIT_EXT |
-		(GVulkanGraphicPipelineLibraryLinkingMode == 1 ? VK_PIPELINE_CREATE_RETAIN_LINK_TIME_OPTIMIZATION_INFO_BIT_EXT : 0);
+	GraphicsPipelineCreateInfo.pNext = &Flags2CreateInfo;
+	GraphicsPipelineCreateInfo.flags = 0; // using VkPipelineCreateFlags2CreateInfo
 	GraphicsPipelineCreateInfo.layout = Device->GetBindlessDescriptorManager()->GetPipelineLayout();
 	GraphicsPipelineCreateInfo.pColorBlendState = &ColorBlendStateCreateInfo;
 	GraphicsPipelineCreateInfo.pMultisampleState = &MultisampleStateCreateInfo;
@@ -1640,16 +1661,24 @@ bool FVulkanPipelineStateCacheManager::CreateGfxPipelineFromLibrary(FVulkanGraph
 	Libraries[2] = PixelShader ? PixelShader->PipelineLibrary : EmptyFragmentStageLibrary;
 	Libraries[3] = FindOrCreateFragmentOutputStateLibrary(GraphicsPipelineState);
 
-	VkPipelineLibraryCreateInfoKHR LibraryCreateInfo;
-	ZeroVulkanStruct(LibraryCreateInfo, VK_STRUCTURE_TYPE_PIPELINE_LIBRARY_CREATE_INFO_KHR);
-	LibraryCreateInfo.libraryCount = Libraries.Num();
-	LibraryCreateInfo.pLibraries = Libraries.GetData();
+	const VkPipelineLibraryCreateInfoKHR LibraryCreateInfo = {
+		.sType = VK_STRUCTURE_TYPE_PIPELINE_LIBRARY_CREATE_INFO_KHR,
+		.pNext = nullptr,
+		.libraryCount = Libraries.Num(),
+		.pLibraries = Libraries.GetData()
+	};
+
+	const VkPipelineCreateFlags2CreateInfo Flags2CreateInfo = {
+		.sType = VK_STRUCTURE_TYPE_PIPELINE_CREATE_FLAGS_2_CREATE_INFO,
+		.pNext = &LibraryCreateInfo,
+		.flags = Device->GetBindlessDescriptorManager()->GetPipelineCreateFlag() |
+				 (GVulkanGraphicPipelineLibraryLinkingMode == 1 ? VK_PIPELINE_CREATE_2_LINK_TIME_OPTIMIZATION_BIT_EXT : 0)
+	};
 
 	VkGraphicsPipelineCreateInfo ExecutablePipelineCreateInfo;
 	ZeroVulkanStruct(ExecutablePipelineCreateInfo, VK_STRUCTURE_TYPE_GRAPHICS_PIPELINE_CREATE_INFO);
-	ExecutablePipelineCreateInfo.pNext = &LibraryCreateInfo;
-	ExecutablePipelineCreateInfo.flags = VK_PIPELINE_CREATE_DESCRIPTOR_BUFFER_BIT_EXT | 
-		(GVulkanGraphicPipelineLibraryLinkingMode == 1 ? VK_PIPELINE_CREATE_LINK_TIME_OPTIMIZATION_BIT_EXT : 0);
+	ExecutablePipelineCreateInfo.pNext = &Flags2CreateInfo;
+	ExecutablePipelineCreateInfo.flags = 0;  // using VkPipelineCreateFlags2CreateInfo
 	ExecutablePipelineCreateInfo.layout = Device->GetBindlessDescriptorManager()->GetPipelineLayout();
 
 	const VkResult Result = VulkanRHI::vkCreateGraphicsPipelines(Device->GetHandle(), nullptr, 1, &ExecutablePipelineCreateInfo, VULKAN_CPU_ALLOCATOR, &GraphicsPipelineState.VulkanPipeline);
@@ -1825,8 +1854,7 @@ bool FVulkanPipelineStateCacheManager::CreateGfxPipelineFromEntry(FVulkanGraphic
 	{
 		ZeroVulkanStruct(PipelineRenderingCreateInfo, VK_STRUCTURE_TYPE_PIPELINE_RENDERING_CREATE_INFO);
 		FillDynamicRenderingCreateInfo(PSO->Desc.RenderTargets, PipelineRenderingCreateInfo, ColorAttachmentFormats);
-		PipelineRenderingCreateInfo.pNext = PipelineInfo.pNext;
-		PipelineInfo.pNext = &PipelineRenderingCreateInfo;
+		AddToPNext(PipelineInfo, PipelineRenderingCreateInfo);
 	}
 	else
 	{
@@ -1873,13 +1901,21 @@ bool FVulkanPipelineStateCacheManager::CreateGfxPipelineFromEntry(FVulkanGraphic
 		PipelineFragmentShadingRate.combinerOps[0] = PipelineToPrimitiveCombinerOperation;
 		PipelineFragmentShadingRate.combinerOps[1] = VK_FRAGMENT_SHADING_RATE_COMBINER_OP_MAX_KHR;		// @todo: This needs to be specified too.
 
-		PipelineFragmentShadingRate.pNext = PipelineInfo.pNext;
-		PipelineInfo.pNext = (void*)&PipelineFragmentShadingRate;
+		AddToPNext(PipelineInfo, PipelineFragmentShadingRate);
 	}
 
+	VkPipelineCreateFlags2CreateInfo Flags2CreateInfo;
 	if (PSO->UsesBindless())
 	{
-		PipelineInfo.flags |= VK_PIPELINE_CREATE_DESCRIPTOR_BUFFER_BIT_EXT;
+		Flags2CreateInfo = {
+			.sType = VK_STRUCTURE_TYPE_PIPELINE_CREATE_FLAGS_2_CREATE_INFO,
+			.pNext = nullptr,
+			.flags = PipelineInfo.flags | Device->GetBindlessDescriptorManager()->GetPipelineCreateFlag()
+		};
+
+		PipelineInfo.flags = 0;
+
+		AddToPNext(PipelineInfo, Flags2CreateInfo);
 	}
 
 	VkResult Result = VK_ERROR_INITIALIZATION_FAILED;
@@ -2906,10 +2942,6 @@ FVulkanComputePipeline* FVulkanPipelineStateCacheManager::CreateComputePipeline(
 		.basePipelineIndex = 0
 	};
 
-	if (Shader->UsesBindless())
-	{
-		PipelineInfo.flags |= VK_PIPELINE_CREATE_DESCRIPTOR_BUFFER_BIT_EXT;
-	}
 
 	VkSpecializationInfo SpecializationInfo;
 	if (SpecializationConstants.Num())
@@ -2946,6 +2978,20 @@ FVulkanComputePipeline* FVulkanPipelineStateCacheManager::CreateComputePipeline(
 		}
 	}
 
+	VkPipelineCreateFlags2CreateInfo Flags2CreateInfo;
+	if (Shader->UsesBindless())
+	{
+		Flags2CreateInfo = {
+			.sType = VK_STRUCTURE_TYPE_PIPELINE_CREATE_FLAGS_2_CREATE_INFO,
+			.pNext = nullptr,
+			.flags = PipelineInfo.flags | Device->GetBindlessDescriptorManager()->GetPipelineCreateFlag()
+		};
+
+		PipelineInfo.flags = 0;
+
+		AddToPNext(PipelineInfo, Flags2CreateInfo);
+	}
+
 	VkResult Result;
 	{
 		QUICK_SCOPE_CYCLE_COUNTER(STAT_VulkanComputePSOCreate);
@@ -2960,7 +3006,6 @@ FVulkanComputePipeline* FVulkanPipelineStateCacheManager::CreateComputePipeline(
 		Pipeline->SetValid(false);
 	}
 
-
 	INC_DWORD_STAT(STAT_VulkanNumPSOs);
 
 	return Pipeline;
diff --git a/Engine/Source/Runtime/VulkanRHI/Private/VulkanShaders.cpp b/Engine/Source/Runtime/VulkanRHI/Private/VulkanShaders.cpp
--- a/Engine/Source/Runtime/VulkanRHI/Private/VulkanShaders.cpp
+++ b/Engine/Source/Runtime/VulkanRHI/Private/VulkanShaders.cpp
@@ -77,25 +77,32 @@ static VkPipeline CreatePreRasterShaderLibrary(FVulkanDevice* Device, FVulkanSha
 	TRACE_CPUPROFILER_EVENT_SCOPE_CONDITIONAL(VulkanMSLibTime, (StageFlag == VK_SHADER_STAGE_MESH_BIT_EXT));
 #endif
 
-	VkShaderModuleCreateInfo ModuleCreateInfo;
-	ZeroVulkanStruct(ModuleCreateInfo, VK_STRUCTURE_TYPE_SHADER_MODULE_CREATE_INFO);
+	const VkGraphicsPipelineLibraryCreateInfoEXT LibraryInfo = {
+		.sType = VK_STRUCTURE_TYPE_GRAPHICS_PIPELINE_LIBRARY_CREATE_INFO_EXT,
+		.pNext = nullptr,
+		.flags = VK_GRAPHICS_PIPELINE_LIBRARY_PRE_RASTERIZATION_SHADERS_BIT_EXT
+	};
 
 	FVulkanShader::FSpirvCode SpirvCode = Shader->GetSpirvCode();
-	ModuleCreateInfo.codeSize = SpirvCode.GetCodeSizeInBytes();
-	ModuleCreateInfo.pCode = SpirvCode.GetCodeView().GetData();
-
-	VkGraphicsPipelineLibraryCreateInfoEXT LibraryInfo;
-	ZeroVulkanStruct(LibraryInfo, VK_STRUCTURE_TYPE_GRAPHICS_PIPELINE_LIBRARY_CREATE_INFO_EXT);
-	LibraryInfo.flags = VK_GRAPHICS_PIPELINE_LIBRARY_PRE_RASTERIZATION_SHADERS_BIT_EXT;
-
-	VkPipelineShaderStageCreateInfo ShaderStageCreateInfo;
-	ZeroVulkanStruct(ShaderStageCreateInfo, VK_STRUCTURE_TYPE_PIPELINE_SHADER_STAGE_CREATE_INFO);
-	ShaderStageCreateInfo.pNext = &ModuleCreateInfo;
-	ShaderStageCreateInfo.stage = StageFlag;
+	const VkShaderModuleCreateInfo ModuleCreateInfo = {
+		.sType = VK_STRUCTURE_TYPE_SHADER_MODULE_CREATE_INFO,
+		.pNext = nullptr,
+		.flags = 0,
+		.codeSize = (size_t)SpirvCode.GetCodeSizeInBytes(),
+		.pCode = SpirvCode.GetCodeView().GetData()
+	};
 
 	ANSICHAR EntryPoint[24];
 	Shader->GetEntryPoint(EntryPoint, 24);
-	ShaderStageCreateInfo.pName = EntryPoint;
+	const VkPipelineShaderStageCreateInfo ShaderStageCreateInfo = {
+		.sType = VK_STRUCTURE_TYPE_PIPELINE_SHADER_STAGE_CREATE_INFO,
+		.pNext = &ModuleCreateInfo,
+		.flags = 0,
+		.stage = StageFlag,
+		.module = VK_NULL_HANDLE, // specified in pNext
+		.pName = EntryPoint,
+		.pSpecializationInfo = nullptr // :todo: create every specialization
+	};
 
 	// See table for pre-raster: https://docs.vulkan.org/guide/latest/dynamic_state_map.html
 	const VkDynamicState PreRasterDynamicStates[] = {
@@ -107,21 +114,31 @@ static VkPipeline CreatePreRasterShaderLibrary(FVulkanDevice* Device, FVulkanSha
 		VK_DYNAMIC_STATE_POLYGON_MODE_EXT  // ext3
 	};
 
-	VkPipelineDynamicStateCreateInfo DynamicStateCreateInfo;
-	ZeroVulkanStruct(DynamicStateCreateInfo, VK_STRUCTURE_TYPE_PIPELINE_DYNAMIC_STATE_CREATE_INFO);
-	DynamicStateCreateInfo.dynamicStateCount = UE_ARRAY_COUNT(PreRasterDynamicStates);
-	DynamicStateCreateInfo.pDynamicStates = PreRasterDynamicStates;
+	const VkPipelineDynamicStateCreateInfo DynamicStateCreateInfo = {
+		.sType = VK_STRUCTURE_TYPE_PIPELINE_DYNAMIC_STATE_CREATE_INFO,
+		.pNext = nullptr,
+		.flags = 0,
+		.dynamicStateCount = UE_ARRAY_COUNT(PreRasterDynamicStates),
+		.pDynamicStates = PreRasterDynamicStates
+	};
 
 	// These struct's important fields are dynamic, set the rest to the default
 	VkPipelineRasterizationStateCreateInfo RasterizerState;
 	FVulkanRasterizerState::ResetCreateInfo(RasterizerState);
 
+	const VkPipelineCreateFlags2CreateInfo Flags2CreateInfo = {
+		.sType = VK_STRUCTURE_TYPE_PIPELINE_CREATE_FLAGS_2_CREATE_INFO,
+		.pNext = &LibraryInfo,
+		.flags = VK_PIPELINE_CREATE_2_LIBRARY_BIT_KHR | 
+				 Device->GetBindlessDescriptorManager()->GetPipelineCreateFlag() |
+				 (GVulkanGraphicPipelineLibraryLinkingMode == 1 ? VK_PIPELINE_CREATE_2_RETAIN_LINK_TIME_OPTIMIZATION_INFO_BIT_EXT : 0)
+	};
+
 	VkGraphicsPipelineCreateInfo GraphicsPipelineCreateInfo;
 	ZeroVulkanStruct(GraphicsPipelineCreateInfo, VK_STRUCTURE_TYPE_GRAPHICS_PIPELINE_CREATE_INFO);
 	GraphicsPipelineCreateInfo.sType = VK_STRUCTURE_TYPE_GRAPHICS_PIPELINE_CREATE_INFO;
-	GraphicsPipelineCreateInfo.pNext = &LibraryInfo;
-	GraphicsPipelineCreateInfo.flags = VK_PIPELINE_CREATE_LIBRARY_BIT_KHR | VK_PIPELINE_CREATE_DESCRIPTOR_BUFFER_BIT_EXT | 
-		(GVulkanGraphicPipelineLibraryLinkingMode == 1 ? VK_PIPELINE_CREATE_RETAIN_LINK_TIME_OPTIMIZATION_INFO_BIT_EXT : 0);
+	GraphicsPipelineCreateInfo.pNext = &Flags2CreateInfo;
+	GraphicsPipelineCreateInfo.flags = 0; // using VkPipelineCreateFlags2CreateInfo
 	GraphicsPipelineCreateInfo.stageCount = 1;
 	GraphicsPipelineCreateInfo.pStages = &ShaderStageCreateInfo;
 	GraphicsPipelineCreateInfo.layout = Device->GetBindlessDescriptorManager()->GetPipelineLayout();
@@ -141,30 +158,42 @@ static VkPipeline CreatePixelShaderLibrary(FVulkanDevice* Device, FVulkanShader*
 	TRACE_CPUPROFILER_EVENT_SCOPE(VulkanPSLibTime);
 #endif
 
-	VkShaderModuleCreateInfo ModuleCreateInfo;
-	ZeroVulkanStruct(ModuleCreateInfo, VK_STRUCTURE_TYPE_SHADER_MODULE_CREATE_INFO);
-
-	FVulkanShader::FSpirvCode SpirvCode = PixelShader->GetSpirvCode();
-	ModuleCreateInfo.codeSize = SpirvCode.GetCodeSizeInBytes();
-	ModuleCreateInfo.pCode = SpirvCode.GetCodeView().GetData();
-
-	VkGraphicsPipelineLibraryCreateInfoEXT LibraryInfo;
-	ZeroVulkanStruct(LibraryInfo, VK_STRUCTURE_TYPE_GRAPHICS_PIPELINE_LIBRARY_CREATE_INFO_EXT);
-	LibraryInfo.flags = VK_GRAPHICS_PIPELINE_LIBRARY_FRAGMENT_SHADER_BIT_EXT;
+	const VkPipelineRenderingCreateInfo RenderingCreateInfo = {
+		.sType = VK_STRUCTURE_TYPE_PIPELINE_RENDERING_CREATE_INFO,
+		.pNext = nullptr,
+		.viewMask = 0,  // :todo-jn: fill viewMask for multiview
+		.colorAttachmentCount = 0, // not used
+		.pColorAttachmentFormats = nullptr, // not used
+		.depthAttachmentFormat = VK_FORMAT_UNDEFINED, // not used
+		.stencilAttachmentFormat = VK_FORMAT_UNDEFINED // not used
+	};
 
-	VkPipelineRenderingCreateInfo RenderingCreateInfo;  
-	ZeroVulkanStruct(RenderingCreateInfo, VK_STRUCTURE_TYPE_PIPELINE_RENDERING_CREATE_INFO);
-	RenderingCreateInfo.viewMask = 0;  // :todo-jn: fill viewMask for multiview
-	LibraryInfo.pNext = &RenderingCreateInfo;
+	const VkGraphicsPipelineLibraryCreateInfoEXT LibraryInfo = {
+		.sType = VK_STRUCTURE_TYPE_GRAPHICS_PIPELINE_LIBRARY_CREATE_INFO_EXT,
+		.pNext = &RenderingCreateInfo,
+		.flags = VK_GRAPHICS_PIPELINE_LIBRARY_FRAGMENT_SHADER_BIT_EXT
+	};
 
-	VkPipelineShaderStageCreateInfo ShaderStageCreateInfo;
-	ZeroVulkanStruct(ShaderStageCreateInfo, VK_STRUCTURE_TYPE_PIPELINE_SHADER_STAGE_CREATE_INFO);
-	ShaderStageCreateInfo.pNext = &ModuleCreateInfo;
-	ShaderStageCreateInfo.stage = VK_SHADER_STAGE_FRAGMENT_BIT;
+	FVulkanShader::FSpirvCode SpirvCode = PixelShader->GetSpirvCode();
+	const VkShaderModuleCreateInfo ModuleCreateInfo = {
+		.sType = VK_STRUCTURE_TYPE_SHADER_MODULE_CREATE_INFO,
+		.pNext = nullptr,
+		.flags = 0,
+		.codeSize = (size_t)SpirvCode.GetCodeSizeInBytes(),
+		.pCode = SpirvCode.GetCodeView().GetData()
+	};
 
 	ANSICHAR EntryPoint[24];
 	PixelShader->GetEntryPoint(EntryPoint, 24);
-	ShaderStageCreateInfo.pName = EntryPoint;
+	const VkPipelineShaderStageCreateInfo ShaderStageCreateInfo = {
+		.sType = VK_STRUCTURE_TYPE_PIPELINE_SHADER_STAGE_CREATE_INFO,
+		.pNext = &ModuleCreateInfo,
+		.flags = 0,
+		.stage = VK_SHADER_STAGE_FRAGMENT_BIT,
+		.module = VK_NULL_HANDLE,  // specified in pNext chain
+		.pName = EntryPoint,
+		.pSpecializationInfo = nullptr // :todo: create every specialization
+	};
 
 	// See table for fragment shader: https://docs.vulkan.org/guide/latest/dynamic_state_map.html
 	const VkDynamicState FragmentShaderDynamicStates[] = {
@@ -180,10 +209,13 @@ static VkPipeline CreatePixelShaderLibrary(FVulkanDevice* Device, FVulkanShader*
 		VK_DYNAMIC_STATE_ALPHA_TO_COVERAGE_ENABLE_EXT // ext3
 	};
 
-	VkPipelineDynamicStateCreateInfo DynamicStateCreateInfo;
-	ZeroVulkanStruct(DynamicStateCreateInfo, VK_STRUCTURE_TYPE_PIPELINE_DYNAMIC_STATE_CREATE_INFO);
-	DynamicStateCreateInfo.dynamicStateCount = UE_ARRAY_COUNT(FragmentShaderDynamicStates);
-	DynamicStateCreateInfo.pDynamicStates = FragmentShaderDynamicStates;
+	const VkPipelineDynamicStateCreateInfo DynamicStateCreateInfo = {
+		.sType = VK_STRUCTURE_TYPE_PIPELINE_DYNAMIC_STATE_CREATE_INFO,
+		.pNext = nullptr,
+		.flags = 0,
+		.dynamicStateCount = UE_ARRAY_COUNT(FragmentShaderDynamicStates),
+		.pDynamicStates = FragmentShaderDynamicStates
+	};
 
 	// These struct's important fields are dynamic, set the rest to the default
 	VkPipelineMultisampleStateCreateInfo MultisampleStateCreateInfo;
@@ -191,12 +223,19 @@ static VkPipeline CreatePixelShaderLibrary(FVulkanDevice* Device, FVulkanShader*
 	VkPipelineDepthStencilStateCreateInfo DepthStencilStateCreateInfo;
 	ZeroVulkanStruct(DepthStencilStateCreateInfo, VK_STRUCTURE_TYPE_PIPELINE_DEPTH_STENCIL_STATE_CREATE_INFO);
 
+	const VkPipelineCreateFlags2CreateInfo Flags2CreateInfo = {
+		.sType = VK_STRUCTURE_TYPE_PIPELINE_CREATE_FLAGS_2_CREATE_INFO,
+		.pNext = &LibraryInfo,
+		.flags = VK_PIPELINE_CREATE_2_LIBRARY_BIT_KHR |
+				 Device->GetBindlessDescriptorManager()->GetPipelineCreateFlag() |
+				 (GVulkanGraphicPipelineLibraryLinkingMode == 1 ? VK_PIPELINE_CREATE_2_RETAIN_LINK_TIME_OPTIMIZATION_INFO_BIT_EXT : 0)
+	};
+
 	VkGraphicsPipelineCreateInfo GraphicsPipelineCreateInfo;
 	ZeroVulkanStruct(GraphicsPipelineCreateInfo, VK_STRUCTURE_TYPE_GRAPHICS_PIPELINE_CREATE_INFO);
 	GraphicsPipelineCreateInfo.sType = VK_STRUCTURE_TYPE_GRAPHICS_PIPELINE_CREATE_INFO;
-	GraphicsPipelineCreateInfo.pNext = &LibraryInfo;
-	GraphicsPipelineCreateInfo.flags = VK_PIPELINE_CREATE_LIBRARY_BIT_KHR | VK_PIPELINE_CREATE_DESCRIPTOR_BUFFER_BIT_EXT |
-		(GVulkanGraphicPipelineLibraryLinkingMode == 1 ? VK_PIPELINE_CREATE_RETAIN_LINK_TIME_OPTIMIZATION_INFO_BIT_EXT : 0);
+	GraphicsPipelineCreateInfo.pNext = &Flags2CreateInfo;
+	GraphicsPipelineCreateInfo.flags = 0; // using VkPipelineCreateFlags2CreateInfo
 	GraphicsPipelineCreateInfo.stageCount = 1;
 	GraphicsPipelineCreateInfo.pStages = &ShaderStageCreateInfo;
 	GraphicsPipelineCreateInfo.pMultisampleState = &MultisampleStateCreateInfo;
