FVulkanQueue::GetQueueType() should be const.

#rnx

#Synced-CL 51937798

diff --git a/Engine/Source/Runtime/VulkanRHI/Private/VulkanQueue.h b/Engine/Source/Runtime/VulkanRHI/Private/VulkanQueue.h
--- a/Engine/Source/Runtime/VulkanRHI/Private/VulkanQueue.h
+++ b/Engine/Source/Runtime/VulkanRHI/Private/VulkanQueue.h
@@ -55,7 +55,7 @@ public:
 	FVulkanQueue(FVulkanDevice& InDevice, uint32 InFamilyIndex, EVulkanQueueType InQueueType);
 	~FVulkanQueue();
 
-	EVulkanQueueType GetQueueType()
+	EVulkanQueueType GetQueueType() const
 	{
 		return QueueType;
 	}
