Report Vulkan OOM as OOMs for crash reporting purposes

[RN] rnx
#Synced-CL 51753244

diff --git a/Engine/Source/Runtime/VulkanRHI/Private/VulkanMemory.cpp b/Engine/Source/Runtime/VulkanRHI/Private/VulkanMemory.cpp
--- a/Engine/Source/Runtime/VulkanRHI/Private/VulkanMemory.cpp
+++ b/Engine/Source/Runtime/VulkanRHI/Private/VulkanMemory.cpp
@@ -7,6 +7,7 @@
 #include "VulkanRHIPrivate.h"
 #include "VulkanMemory.h"
 #include "Misc/OutputDeviceRedirector.h"
+#include "HAL/PlatformMemory.h"
 #include "HAL/PlatformStackWalk.h"
 #include "VulkanContext.h"
 #include "VulkanLLM.h"
@@ -3009,6 +3010,7 @@ namespace VulkanRHI
 			if (!ResourceTypeHeaps[TypeIndex]->AllocateResource(OutAllocation, AllocationOwner, EType::Image, MemoryReqs.size, MemoryReqs.alignment, bMapped, bForceSeparateAllocation, MetaType, bExternal, File, Line))
 			{
 				DumpMemory();
+				FPlatformMemory::bIsOOM = true;
 				UE_LOG(LogVulkanRHI, Fatal, TEXT("Out Of Memory, trying to allocate %d bytes\n"), MemoryReqs.size);
 				return false;
 			}
@@ -3076,6 +3078,7 @@ namespace VulkanRHI
 			if (!ResourceTypeHeaps[TypeIndex]->AllocateResource(OutAllocation, AllocationOwner, EType::Buffer, MemoryReqs.size, MemoryReqs.alignment, bMapped, bForceSeparateAllocation, MetaType, bExternal, File, Line))
 			{
 				DumpMemory();
+				FPlatformMemory::bIsOOM = true;
 				UE_LOG(LogVulkanRHI, Fatal, TEXT("Out Of Memory, trying to allocate %d bytes\n"), MemoryReqs.size);
 				return false;
 			}
@@ -3693,6 +3696,7 @@ namespace VulkanRHI
 			DumpRenderTargetPoolMemory(*GLog);
 			GLog->Flush();
 
+			FPlatformMemory::bIsOOM = true;
 			UE_LOG(LogVulkanRHI, Fatal, TEXT("Out of %s Memory, Requested%.2fKB MemTypeIndex=%d\n"), MemoryType, float(AllocationSize), MemoryTypeIndex);
 		}
 	}
