VulkanRHI: Remove unused RenderOffscreen inside the RHI, it's handled at a higher level now.

#rnx

#Synced-CL 51908825

diff --git a/Engine/Source/Runtime/VulkanRHI/Private/VulkanViewport.cpp b/Engine/Source/Runtime/VulkanRHI/Private/VulkanViewport.cpp
--- a/Engine/Source/Runtime/VulkanRHI/Private/VulkanViewport.cpp
+++ b/Engine/Source/Runtime/VulkanRHI/Private/VulkanViewport.cpp
@@ -109,7 +109,6 @@ FVulkanViewport::FVulkanViewport(FVulkanDevice& InDevice, void* InWindowHandle,
 	, SwapChain(nullptr)
 	, WindowHandle(InWindowHandle)
 	, PresentCount(0)
-	, bRenderOffscreen(false)
 	, AcquiredSemaphore(nullptr)
 {
 	check(IsInGameThread());
@@ -122,8 +121,6 @@ FVulkanViewport::FVulkanViewport(FVulkanDevice& InDevice, void* InWindowHandle,
 	// Make sure Instance is created
 	FVulkanDynamicRHI::Get().InitInstance();
 
-	bRenderOffscreen = FParse::Param(FCommandLine::Get(), TEXT("RenderOffScreen"));
-
 	FVulkanPlatformWindowContext WindowContext(InWindowHandle);
 
 	ENQUEUE_RENDER_COMMAND(CreateSwapchain)(
@@ -832,7 +829,7 @@ bool FVulkanViewport::Present(FVulkanCommandListContext& Context, const FRHIPres
 
 bool FVulkanViewport::SupportsStandardSwapchain()
 {
-	return !bRenderOffscreen && !FVulkanDynamicRHI::Get().bIsStandaloneStereoDevice;
+	return !FVulkanDynamicRHI::Get().bIsStandaloneStereoDevice;
 }
 
 bool FVulkanViewport::RequiresRenderingBackBuffer()
@@ -842,7 +839,7 @@ bool FVulkanViewport::RequiresRenderingBackBuffer()
 
 EPixelFormat FVulkanViewport::GetPixelFormatForNonDefaultSwapchain()
 {
-	if (bRenderOffscreen || FVulkanDynamicRHI::Get().bIsStandaloneStereoDevice)
+	if (FVulkanDynamicRHI::Get().bIsStandaloneStereoDevice)
 	{
 		return PF_R8G8B8A8;
 	}
diff --git a/Engine/Source/Runtime/VulkanRHI/Private/VulkanViewport.h b/Engine/Source/Runtime/VulkanRHI/Private/VulkanViewport.h
--- a/Engine/Source/Runtime/VulkanRHI/Private/VulkanViewport.h
+++ b/Engine/Source/Runtime/VulkanRHI/Private/VulkanViewport.h
@@ -134,7 +134,6 @@ protected:
 	FVulkanSwapChain* SwapChain;
 	void* WindowHandle;
 	uint32 PresentCount;
-	bool bRenderOffscreen;
 
 	int8 LockToVsync;
 
