VulkanRHI:  Create a BindlessResourceSet value to represent the single resource heap's set.

#rnx

#Synced-CL 52071393

diff --git a/Engine/Source/Runtime/VulkanRHI/Public/VulkanCommon.h b/Engine/Source/Runtime/VulkanRHI/Public/VulkanCommon.h
--- a/Engine/Source/Runtime/VulkanRHI/Public/VulkanCommon.h
+++ b/Engine/Source/Runtime/VulkanRHI/Public/VulkanCommon.h
@@ -100,17 +100,18 @@ namespace VulkanBindless
 	enum EDescriptorSets
 	{
 		BindlessSamplerSet = 0,
+		BindlessResourceSet = 1,
 
-		BindlessStorageBufferSet = 1,
-		BindlessUniformBufferSet = 1,
+		BindlessStorageBufferSet = BindlessResourceSet,
+		BindlessUniformBufferSet = BindlessResourceSet,
 
-		BindlessStorageImageSet = 1,
-		BindlessSampledImageSet = 1,
+		BindlessStorageImageSet = BindlessResourceSet,
+		BindlessSampledImageSet = BindlessResourceSet,
 
-		BindlessStorageTexelBufferSet = 1,
-		BindlessUniformTexelBufferSet = 1,
+		BindlessStorageTexelBufferSet = BindlessResourceSet,
+		BindlessUniformTexelBufferSet = BindlessResourceSet,
 
-		BindlessAccelerationStructureSet = 1,
+		BindlessAccelerationStructureSet = BindlessResourceSet,
 
 		// Number of sets reserved for samplers/resources
 		NumBindlessSets,
