Android: Enable loading of the standard Vulkan binary cache (typically saved on game exit). This was previously disabled under the assumption that the PSO file cache would be used.
The binary cache stores all PSOs by default, or all compute PSOs when the chunked PSO cache is enabled.

#Synced-CL 52121382

diff --git a/Engine/Source/Runtime/VulkanRHI/Private/VulkanPipeline.cpp b/Engine/Source/Runtime/VulkanRHI/Private/VulkanPipeline.cpp
--- a/Engine/Source/Runtime/VulkanRHI/Private/VulkanPipeline.cpp
+++ b/Engine/Source/Runtime/VulkanRHI/Private/VulkanPipeline.cpp
@@ -100,10 +100,6 @@ static TAutoConsoleVariable<int32> CVarLRUPipelineCapacity(
 	TEXT("Maximum no. of PSOs in LRU."),
 	ECVF_RenderThreadSafe| ECVF_ReadOnly);
 
-
-
-
-
 static TAutoConsoleVariable<int32> GEnablePipelineCacheLoadCvar(
 	TEXT("r.Vulkan.PipelineCacheLoad"),
 	1,
@@ -111,15 +107,6 @@ static TAutoConsoleVariable<int32> GEnablePipelineCacheLoadCvar(
 	TEXT("1 to enable using pipeline cache.\n")
 );
 
-static TAutoConsoleVariable<int32> GPipelineCacheFromShaderPipelineCacheCvar(
-	TEXT("r.Vulkan.PipelineCacheFromShaderPipelineCache"),
-	PLATFORM_ANDROID,
-	TEXT("0 look for a pipeline cache in the normal locations with the normal names.\n")
-	TEXT("1 tie the vulkan pipeline cache to the shader pipeline cache, use the PSOFC guid as part of the filename, etc.\n"),
-	ECVF_ReadOnly
-);
-
-
 static int32 GEnablePipelineCacheCompression = 1;
 static FAutoConsoleVariableRef GEnablePipelineCacheCompressionCvar(
 	TEXT("r.Vulkan.PipelineCacheCompression"),
@@ -497,32 +484,39 @@ bool FVulkanPipelineStateCacheManager::Load(const TArray<FString>& CacheFilename
 
 void FVulkanPipelineStateCacheManager::InitAndLoad(const TArray<FString>& CacheFilenames)
 {
+	// Where PSOFC cache is stored
+	CompiledPSOCacheTopFolderPath = FVulkanPlatform::GetCompiledPSOCacheTopFolderPath();
+
+	// Remove pipeline cache files on disk if either -ClearVulkanBinaryProgramCache or clearPSODriverCache are specified on command line
+	if (FParse::Param(FCommandLine::Get(), TEXT("ClearVulkanBinaryProgramCache")) || 
+		FParse::Param(FCommandLine::Get(), TEXT("clearPSODriverCache")))
+	{
+		UE_LOG(LogVulkanRHI, Log, TEXT("FVulkanPipelineStateCacheManager: Clearing pipeline cache on disk"));
+		for (const FString& CacheFilename : CacheFilenames)
+		{
+			FString BinaryCacheFilename = FVulkanPlatform::CreatePSOBinaryCacheFilename(Device, CacheFilename);
+			FPlatformFileManager::Get().GetPlatformFile().DeleteFile(*BinaryCacheFilename);
+		}
+		FPlatformFileManager::Get().GetPlatformFile().DeleteDirectoryRecursively(*CompiledPSOCacheTopFolderPath);
+	}
+	
 	if (GEnablePipelineCacheLoadCvar.GetValueOnAnyThread() == 0)
 	{
 		UE_LOG(LogVulkanRHI, Display, TEXT("Not loading pipeline cache per r.Vulkan.PipelineCacheLoad=0"));
 	}
 	else
 	{
-		if (GPipelineCacheFromShaderPipelineCacheCvar.GetValueOnAnyThread() == 0)
+		if (!FPipelineFileCacheManager::IsPipelineFileCacheEnabled())
 		{
+			// Load the default PSO cache from disk.
+			// Stores all compute PSOs, and graphics PSOs when the PSOFC and chunked PSO cache are disabled
 			Load(CacheFilenames, GlobalPSOCache);
-		}
-		else
-		{
-			UE_LOG(LogVulkanRHI, Log, TEXT("FVulkanPipelineStateCacheManager will check for loading, etc when ShaderPipelineCache opens its file"));
-
-			CompiledPSOCacheTopFolderPath = FVulkanPlatform::GetCompiledPSOCacheTopFolderPath();
 
-			// Remove entire ProgramBinaryCache folder if either -ClearVulkanBinaryProgramCache or clearPSODriverCache are specified on command line
-			if (FParse::Param(FCommandLine::Get(), TEXT("ClearVulkanBinaryProgramCache")) || FParse::Param(FCommandLine::Get(), TEXT("clearPSODriverCache")))
-			{
-				UE_LOG(LogVulkanRHI, Log, TEXT("FVulkanPipelineStateCacheManager: Deleting binary program cache folder for -ClearVulkanBinaryProgramCache: %s"), *CompiledPSOCacheTopFolderPath);
-				FPlatformFileManager::Get().GetPlatformFile().DeleteDirectoryRecursively(*CompiledPSOCacheTopFolderPath);
-			}
-
-			OnShaderPipelineCacheOpenedDelegate = FShaderPipelineCache::GetCacheOpenedDelegate().AddRaw(this, &FVulkanPipelineStateCacheManager::OnShaderPipelineCacheOpened);
-			OnShaderPipelineCachePrecompilationCompleteDelegate = FShaderPipelineCache::GetPrecompilationCompleteDelegate().AddRaw(this, &FVulkanPipelineStateCacheManager::OnShaderPipelineCachePrecompilationComplete);
+			//TODO: Since CreateComputePipelineFromShader uses only GlobalPSOCache, PSOFC will always miss CS PSOs
 		}
+
+		OnShaderPipelineCacheOpenedDelegate = FShaderPipelineCache::GetCacheOpenedDelegate().AddRaw(this, &FVulkanPipelineStateCacheManager::OnShaderPipelineCacheOpened);
+		OnShaderPipelineCachePrecompilationCompleteDelegate = FShaderPipelineCache::GetPrecompilationCompleteDelegate().AddRaw(this, &FVulkanPipelineStateCacheManager::OnShaderPipelineCachePrecompilationComplete);
 	}
 
 	FScopedPipelineCache PipelineCacheExclusive = GlobalPSOCache.Get(EPipelineCacheAccess::Exclusive);
@@ -537,7 +531,10 @@ void FVulkanPipelineStateCacheManager::InitAndLoad(const TArray<FString>& CacheF
 
 void FVulkanPipelineStateCacheManager::Save(const FString& CacheFilename)
 {
-	SavePSOCache(CacheFilename, GlobalPSOCache);
+	if (!FPipelineFileCacheManager::IsPipelineFileCacheEnabled())
+	{
+		SavePSOCache(CacheFilename, GlobalPSOCache);
+	}
 
 	//TODO: Save LRU cache here
 }
