VulkanRHI: first code drop to add VK_EXT_descriptor_heap support to the bindless manager (WIP).  Keep descriptor buffers by default and make it selectable.  Replace the BindlessSetStates array with dedicated Sampler and Resource states.

#Synced-CL 52253672

diff --git a/Engine/Source/Runtime/VulkanRHI/Private/VulkanBindlessDescriptorManager.cpp b/Engine/Source/Runtime/VulkanRHI/Private/VulkanBindlessDescriptorManager.cpp
--- a/Engine/Source/Runtime/VulkanRHI/Private/VulkanBindlessDescriptorManager.cpp
+++ b/Engine/Source/Runtime/VulkanRHI/Private/VulkanBindlessDescriptorManager.cpp
@@ -14,7 +14,7 @@ static FAutoConsoleVariableRef CVarVulkanBindlessMaxSamplerDescriptorCount(
 	ECVF_ReadOnly
 );
 
-int32 GVulkanBindlessMaxResourceDescriptorCount = 1024 * 1024;
+int32 GVulkanBindlessMaxResourceDescriptorCount = 768 * 1024;
 static FAutoConsoleVariableRef CVarVulkanBindlessMaxResourceDescriptorCount(
 	TEXT("r.Vulkan.Bindless.MaxResourceDescriptorCount"),
 	GVulkanBindlessMaxResourceDescriptorCount,
@@ -22,6 +22,17 @@ static FAutoConsoleVariableRef CVarVulkanBindlessMaxResourceDescriptorCount(
 	ECVF_ReadOnly
 );
 
+int32 GVulkanBindlessPreferredExtension = 0;
+static FAutoConsoleVariableRef CVarVulkanBindlessPreferredExtension(
+	TEXT("r.Vulkan.Bindless.PreferredExtension"),
+	GVulkanBindlessPreferredExtension,
+	TEXT("Preferred extension to use for bindless:\n")
+	TEXT("0: VK_EXT_descriptor_buffer (default)\n")
+	TEXT("1: VK_EXT_descriptor_heap\n"),
+	ECVF_ReadOnly
+);
+
+
 DECLARE_STATS_GROUP(TEXT("Vulkan Bindless"), STATGROUP_VulkanBindless, STATCAT_Advanced);
 
 DECLARE_DWORD_ACCUMULATOR_STAT(TEXT("Sampler Descriptors Allocated"),               STAT_VulkanBindlessSamplerDescriptorsAllocated,               STATGROUP_VulkanBindless);
@@ -166,87 +177,139 @@ static inline VkMemoryPropertyFlags GetDescriptorBufferMemoryType(FVulkanDevice&
 inline FVulkanBindlessDescriptorManager::BindlessSetState& FVulkanBindlessDescriptorManager::GetBindlessState(ERHIDescriptorType RHIDescriptorType)
 {
 	check(RHIDescriptorType != ERHIDescriptorType::Invalid);
-	const VkDescriptorType DescriptorType = GetVkDescriptorTypeFromUEDescriptorType(RHIDescriptorType);
-	return GetBindlessState(DescriptorType);
+	return (RHIDescriptorType == ERHIDescriptorType::Sampler) ? SamplerState : ResourceState;
 }
 
-inline FVulkanBindlessDescriptorManager::BindlessSetState& FVulkanBindlessDescriptorManager::GetBindlessState(VkDescriptorType DescriptorType)
+// Check all the requirements to be running in Bindless using Descriptor Heaps
+static FVulkanBindlessDescriptorManager::EBindlessType VerifyDescriptorHeapSupport(FVulkanDevice& Device)
 {
-	check(DescriptorType != VK_DESCRIPTOR_TYPE_MAX_ENUM);
-	for (BindlessSetState& State : BindlessSetStates)
+	const VkPhysicalDeviceProperties& GpuProps = Device.GetDeviceProperties();
+	const FOptionalVulkanDeviceExtensions& OptionalDeviceExtensions = Device.GetOptionalExtensions();
+	const VkPhysicalDeviceDescriptorHeapPropertiesEXT& DescriptorHeapProperties = Device.GetOptionalExtensionProperties().DescriptorHeapProps;
+
+	const bool bMeetsExtensionsRequirements =
+		OptionalDeviceExtensions.HasEXTDescriptorIndexing &&
+		OptionalDeviceExtensions.HasBufferDeviceAddress &&
+		OptionalDeviceExtensions.HasEXTDescriptorHeap &&
+		OptionalDeviceExtensions.HasKHRMaintenance5;
+
+	if (bMeetsExtensionsRequirements)
 	{
-		if (State.DescriptorTypes.Contains(DescriptorType))
+		const uint32 ResourceDescriptorSize = FMath::Max(DescriptorHeapProperties.bufferDescriptorSize, DescriptorHeapProperties.imageDescriptorSize);
+		const bool bMeetsPropertiesRequirements =
+			// Our DXC does not yet support SamplerHeapEXT/ResourceHeapEXT builtins, 
+			// so we bind the sampler and descriptor heaps as regular bindings in their own descriptor sets.
+			// We still use an additional descriptor set for "single use" UBs, 
+			// push constants with indices into the resource heap will be used to maintain compat with the push descriptors used by descriptor_buffers
+			// (this might change if/when support for descriptor_buffers is dropped)
+			(GpuProps.limits.maxBoundDescriptorSets >= 3) &&
+			(DescriptorHeapProperties.maxResourceHeapSize >= GVulkanBindlessMaxResourceDescriptorCount * ResourceDescriptorSize + DescriptorHeapProperties.minResourceHeapReservedRange) &&
+			(DescriptorHeapProperties.maxSamplerHeapSize >= GVulkanBindlessMaxSamplerDescriptorCount * DescriptorHeapProperties.samplerDescriptorSize + DescriptorHeapProperties.minSamplerHeapReservedRange) &&
+			Device.GetDeviceMemoryManager().SupportsMemoryType(GetDescriptorBufferMemoryType(Device));
+
+		if (bMeetsPropertiesRequirements)
+		{
+			UE_LOG(LogRHI, Display, TEXT("Bindless will use VK_EXT_descriptor_heap extension."));
+			return FVulkanBindlessDescriptorManager::EBindlessType::DescriptorHeap;
+		}
+		else
 		{
-			return State;
+			UE_LOG(LogRHI, Warning, TEXT("Bindless descriptor were requested but NOT enabled because of insufficient property support."));
 		}
 	}
-	checkf(false, TEXT("Bindless state does not exist for descriptor type %s!"), VK_TYPE_TO_STRING(VkDescriptorType, DescriptorType));
-	checkNoEntry();
-	return BindlessSetStates[0];
+	else
+	{
+		UE_LOG(LogRHI, Warning, TEXT("Bindless descriptor were requested but NOT enabled because of missing extension support."));
+	}
+
+	return FVulkanBindlessDescriptorManager::EBindlessType::Disabled;
 }
 
 // Check all the requirements to be running in Bindless using Descriptor Buffers
-static bool VerifySupport(FVulkanDevice& Device, ERHIBindlessConfiguration BindlessConfig)
+static FVulkanBindlessDescriptorManager::EBindlessType VerifyDescriptorBufferSupport(FVulkanDevice& Device)
 {
-	if (BindlessConfig == ERHIBindlessConfiguration::Disabled)
+	const VkPhysicalDeviceProperties& GpuProps = Device.GetDeviceProperties();
+	const FOptionalVulkanDeviceExtensions& OptionalDeviceExtensions = Device.GetOptionalExtensions();
+	const VkPhysicalDeviceDescriptorBufferPropertiesEXT& DescriptorBufferProperties = Device.GetOptionalExtensionProperties().DescriptorBufferProps;
+
+	const bool bMeetsExtensionsRequirements =
+		OptionalDeviceExtensions.HasEXTDescriptorIndexing &&
+		OptionalDeviceExtensions.HasBufferDeviceAddress &&
+		OptionalDeviceExtensions.HasEXTDescriptorBuffer &&
+		OptionalDeviceExtensions.HasEXTMutableDescriptorType &&
+		OptionalDeviceExtensions.HasKHRMaintenance5;
+
+	if (bMeetsExtensionsRequirements)
 	{
-		return false;
+		// Requirements (lines up with descriptor sets and bindings from shader compilation):
+		// - one descriptor set for sampler
+		// - one descriptor set for resources
+		// - one descriptor set for "single use" descriptor buffers
+		// We bind one descriptor buffer for each descriptor set, so we need to support a minimum of 3 for each
+		const bool bMeetsPropertiesRequirements =
+			(GpuProps.limits.maxBoundDescriptorSets >= 3) &&
+			(DescriptorBufferProperties.maxDescriptorBufferBindings >= 3) &&
+			(DescriptorBufferProperties.maxResourceDescriptorBufferBindings >= 1) &&
+			(DescriptorBufferProperties.maxSamplerDescriptorBufferBindings >= 1) &&
+			Device.GetDeviceMemoryManager().SupportsMemoryType(GetDescriptorBufferMemoryType(Device));
+
+		if (bMeetsPropertiesRequirements)
+		{
+			UE_LOG(LogRHI, Display, TEXT("Bindless will use VK_EXT_descriptor_buffer extension."));
+			return FVulkanBindlessDescriptorManager::EBindlessType::DescriptorBuffer;
+		}
+		else
+		{
+			UE_LOG(LogRHI, Warning, TEXT("Bindless descriptor were requested but NOT enabled because of insufficient property support."));
+		}
 	}
-
-	if (BindlessConfig > ERHIBindlessConfiguration::RayTracing)
+	else
 	{
-		const VkPhysicalDeviceProperties& GpuProps = Device.GetDeviceProperties();
-		const FOptionalVulkanDeviceExtensions& OptionalDeviceExtensions = Device.GetOptionalExtensions();
-		const VkPhysicalDeviceDescriptorBufferPropertiesEXT& DescriptorBufferProperties = Device.GetOptionalExtensionProperties().DescriptorBufferProps;
-
-		const bool bMeetsExtensionsRequirements =
-			OptionalDeviceExtensions.HasEXTDescriptorIndexing &&
-			OptionalDeviceExtensions.HasBufferDeviceAddress &&
-			OptionalDeviceExtensions.HasEXTDescriptorBuffer &&
-			OptionalDeviceExtensions.HasEXTMutableDescriptorType;
+		UE_LOG(LogRHI, Warning, TEXT("Bindless descriptor were requested but NOT enabled because of missing extension support."));
+	}
 
-		if (bMeetsExtensionsRequirements)
-		{
-			const bool bMeetsPropertiesRequirements =
-				(GpuProps.limits.maxBoundDescriptorSets >= 3) &&
-				(DescriptorBufferProperties.maxDescriptorBufferBindings >= 3) &&
-				(DescriptorBufferProperties.maxResourceDescriptorBufferBindings >= 1) &&
-				(DescriptorBufferProperties.maxSamplerDescriptorBufferBindings >= 1) &&
-				Device.GetDeviceMemoryManager().SupportsMemoryType(GetDescriptorBufferMemoryType(Device));
+	return FVulkanBindlessDescriptorManager::EBindlessType::Disabled;
+}
 
-			if (bMeetsPropertiesRequirements)
-			{
-				extern TAutoConsoleVariable<int32> GDynamicGlobalUBs;
-				if (GDynamicGlobalUBs->GetInt() != 0)
-				{
-					UE_LOG(LogRHI, Warning, TEXT("Dynamic Uniform Buffers are enabled, but they will not be used with Vulkan bindless."));
-				}
+// Check all the requirements and select the descriptor extension to use (buffer vs heap)
+static FVulkanBindlessDescriptorManager::EBindlessType VerifySupport(FVulkanDevice& Device, ERHIBindlessConfiguration BindlessConfig)
+{
+	FVulkanBindlessDescriptorManager::EBindlessType BindlessType = FVulkanBindlessDescriptorManager::EBindlessType::Disabled;
 
-				extern int32 GVulkanEnableDefrag;
-				if (GVulkanEnableDefrag != 0)  // :todo-jn: to be turned back on with new defragger
-				{
-					UE_LOG(LogRHI, Warning, TEXT("Memory defrag is enabled, but it will not be used with Vulkan bindless."));
-					GVulkanEnableDefrag = 0;
-				}
+	if (BindlessConfig != ERHIBindlessConfiguration::Disabled)
+	{
+		if (BindlessConfig > ERHIBindlessConfiguration::RayTracing)
+		{
+			BindlessType = (GVulkanBindlessPreferredExtension == 0) ? VerifyDescriptorBufferSupport(Device) : VerifyDescriptorHeapSupport(Device);
 
-				return true;
-			}
-			else
+			if (BindlessType == FVulkanBindlessDescriptorManager::EBindlessType::Disabled)
 			{
-				UE_LOG(LogRHI, Warning, TEXT("Bindless descriptor were requested but NOT enabled because of insufficient property support."));
+				BindlessType = (GVulkanBindlessPreferredExtension == 0) ? VerifyDescriptorHeapSupport(Device) : VerifyDescriptorBufferSupport(Device);
 			}
 		}
 		else
 		{
-			UE_LOG(LogRHI, Warning, TEXT("Bindless descriptor were requested but NOT enabled because of missing extension support."));
+			UE_LOG(LogRHI, Warning, TEXT("Bindless in Vulkan must currently be fully enabled (all samplers and resources) or fully disabled."));
 		}
 	}
-	else
+
+	if (BindlessType != FVulkanBindlessDescriptorManager::EBindlessType::Disabled)
 	{
-		UE_LOG(LogRHI, Warning, TEXT("Bindless in Vulkan must currently be fully enabled (all samplers and resources) or fully disabled."));
+		extern TAutoConsoleVariable<int32> GDynamicGlobalUBs;
+		if (GDynamicGlobalUBs->GetInt() != 0)
+		{
+			UE_LOG(LogRHI, Warning, TEXT("Dynamic Uniform Buffers are enabled, but they will not be used with Vulkan bindless."));
+		}
+
+		extern int32 GVulkanEnableDefrag;
+		if (GVulkanEnableDefrag != 0)  // :todo-jn: to be turned back on with new defragger
+		{
+			UE_LOG(LogRHI, Warning, TEXT("Memory defrag is enabled, but it will not be used with Vulkan bindless."));
+			GVulkanEnableDefrag = 0;
+		}
 	}
 
-	return false;
+	return BindlessType;
 }
 
 
@@ -269,12 +332,247 @@ void FVulkanBindlessDescriptorManager::Init()
 {
 	Configuration = UE::RHICore::GetBindlessConfigurationOnStartup(GMaxRHIShaderPlatform);
 
-	if (!VerifySupport(Device, Configuration))
+	BindlessType = VerifySupport(Device, Configuration);
+
+	switch (BindlessType)
 	{
+	case FVulkanBindlessDescriptorManager::EBindlessType::DescriptorHeap:
+		InitDescriptorHeaps();
+		break;
+	case FVulkanBindlessDescriptorManager::EBindlessType::DescriptorBuffer:
+		InitDescriptorBuffers();
+		break;
+	case FVulkanBindlessDescriptorManager::EBindlessType::Disabled:
 		Configuration = ERHIBindlessConfiguration::Disabled;
-		return;
+		break;
+	}
+}
+
+void FVulkanBindlessDescriptorManager::InitDescriptorHeaps()
+{
+	const VkDevice DeviceHandle = Device.GetHandle();
+	const VkPhysicalDeviceDescriptorHeapPropertiesEXT& DescriptorHeapProperties = Device.GetOptionalExtensionProperties().DescriptorHeapProps;
+	const bool bHasRaytracingExtensions = Device.GetOptionalExtensions().HasRaytracingExtensions();
+
+	auto CreateDescriptorHeap = [&Device = Device](uint32 HeapSize, uint32 ReservedSize, BindlessSetState& InOutState) -> VkBindHeapInfoEXT {
+
+		const VkBufferUsageFlags BufferUsageFlags = VK_BUFFER_USAGE_DESCRIPTOR_HEAP_BIT_EXT | VK_BUFFER_USAGE_SHADER_DEVICE_ADDRESS_BIT | VK_BUFFER_USAGE_TRANSFER_DST_BIT;
+
+		// Create buffer for heap
+		InOutState.BufferHandle = Device.CreateBuffer(HeapSize, BufferUsageFlags);
+
+		// Allocate memory for the buffer and bind it
+		{
+			const VkDevice DeviceHandle = Device.GetHandle();
+
+			VkMemoryRequirements BufferMemoryReqs;
+			FMemory::Memzero(BufferMemoryReqs);
+			VulkanRHI::vkGetBufferMemoryRequirements(DeviceHandle, InOutState.BufferHandle, &BufferMemoryReqs);
+			check(BufferMemoryReqs.size >= HeapSize);
+
+			uint32 MemoryTypeIndex = 0;
+			VERIFYVULKANRESULT(Device.GetDeviceMemoryManager().GetMemoryTypeFromProperties(BufferMemoryReqs.memoryTypeBits, GetDescriptorBufferMemoryType(Device), &MemoryTypeIndex));
+
+			const VkMemoryAllocateFlagsInfo FlagsInfo = {
+				.sType = VK_STRUCTURE_TYPE_MEMORY_ALLOCATE_FLAGS_INFO,
+				.pNext = nullptr,
+				.flags = VK_MEMORY_ALLOCATE_DEVICE_ADDRESS_BIT,
+				.deviceMask = 0
+			};
+
+			const VkMemoryAllocateInfo AllocateInfo = {
+				.sType = VK_STRUCTURE_TYPE_MEMORY_ALLOCATE_INFO,
+				.pNext = &FlagsInfo,
+				.allocationSize = BufferMemoryReqs.size,
+				.memoryTypeIndex = MemoryTypeIndex
+			};
+
+			VERIFYVULKANRESULT(VulkanRHI::vkAllocateMemory(DeviceHandle, &AllocateInfo, VULKAN_CPU_ALLOCATOR, &InOutState.MemoryHandle));
+			VERIFYVULKANRESULT(VulkanRHI::vkBindBufferMemory(DeviceHandle, InOutState.BufferHandle, InOutState.MemoryHandle, 0));
+			VERIFYVULKANRESULT(VulkanRHI::vkMapMemory(DeviceHandle, InOutState.MemoryHandle, 0, VK_WHOLE_SIZE, 0, (void**)&InOutState.MappedPointer));
+			FMemory::Memzero(InOutState.MappedPointer, AllocateInfo.allocationSize);
+		}
+
+		// Return the heap info with the buffer's address
+		const VkBufferDeviceAddressInfo AddressInfo = {
+			.sType = VK_STRUCTURE_TYPE_BUFFER_DEVICE_ADDRESS_INFO,
+			.pNext = nullptr,
+			.buffer = InOutState.BufferHandle
+		};
+		const VkDeviceAddress BufferAddress = VulkanRHI::vkGetBufferDeviceAddressKHR(Device.GetHandle(), &AddressInfo);
+
+		return 
+		{
+			.sType = VK_STRUCTURE_TYPE_BIND_HEAP_INFO_EXT,
+			.pNext = nullptr,
+			.heapRange = {
+				.address = BufferAddress,
+				.size = HeapSize,
+			},
+			.reservedRangeOffset = HeapSize - ReservedSize,
+			.reservedRangeSize = ReservedSize,
+		};
+	};
+
+	// Create the sampler heap
+	{
+		SamplerState.DescriptorSize = DescriptorHeapProperties.samplerDescriptorSize;
+
+		const uint32 SamplerHeapSize = GVulkanBindlessMaxSamplerDescriptorCount * SamplerState.DescriptorSize + DescriptorHeapProperties.minSamplerHeapReservedRange;
+		SamplerBindHeapInfo = CreateDescriptorHeap(SamplerHeapSize, DescriptorHeapProperties.minSamplerHeapReservedRange, SamplerState);
+
+		TArray<TStatId, TInlineAllocator<1>> AllocatorStats;
+		AllocatorStats.Emplace(GET_STATID(STAT_BindlessSamplerDescriptorsAllocated));
+		SamplerState.Allocator = new FRHIHeapDescriptorAllocator(ERHIDescriptorTypeMask::Sampler, GVulkanBindlessMaxSamplerDescriptorCount, TConstArrayView<TStatId>(AllocatorStats));
 	}
+	
+	// Create the resource heap
+	{
+		ResourceState.DescriptorSize = FMath::Max(DescriptorHeapProperties.bufferDescriptorSize, DescriptorHeapProperties.imageDescriptorSize);
+
+		const uint32 ResourceHeapSize = GVulkanBindlessMaxResourceDescriptorCount * ResourceState.DescriptorSize + DescriptorHeapProperties.minResourceHeapReservedRange;
+		ResourceBindHeapInfo = CreateDescriptorHeap(ResourceHeapSize, DescriptorHeapProperties.minResourceHeapReservedRange, ResourceState);
+
+		TArray<TStatId, TInlineAllocator<1>> AllocatorStats;
+		AllocatorStats.Emplace(GET_STATID(STAT_BindlessResourceDescriptorsAllocated));
+		ERHIDescriptorTypeMask TypeMask{};
+		EnumAddFlags(TypeMask, ERHIDescriptorTypeMask::TextureSRV);
+		EnumAddFlags(TypeMask, ERHIDescriptorTypeMask::TextureUAV);
+		EnumAddFlags(TypeMask, ERHIDescriptorTypeMask::TypedBufferSRV);
+		EnumAddFlags(TypeMask, ERHIDescriptorTypeMask::TypedBufferUAV);
+		EnumAddFlags(TypeMask, ERHIDescriptorTypeMask::BufferUAV); // VK_DESCRIPTOR_TYPE_STORAGE_BUFFER
+		EnumAddFlags(TypeMask, ERHIDescriptorTypeMask::BufferSRV | ERHIDescriptorTypeMask::CBV); // VK_DESCRIPTOR_TYPE_UNIFORM_BUFFER
+		if (bHasRaytracingExtensions)
+		{
+			EnumAddFlags(TypeMask, ERHIDescriptorTypeMask::AccelerationStructure); // VK_DESCRIPTOR_TYPE_ACCELERATION_STRUCTURE_KHR
+		}
+
+		ResourceState.Allocator = new FRHIHeapDescriptorAllocator(TypeMask, GVulkanBindlessMaxResourceDescriptorCount, TConstArrayView<TStatId>(AllocatorStats));
+	}
+
+	// Create the binding mappings
+	{
+		// Setup a common mapping for our sampler and resources heaps (will be inserted in all arrays)
+		const VkDescriptorSetAndBindingMappingEXT SamplerHeapMapping = {
+			.sType = VK_STRUCTURE_TYPE_DESCRIPTOR_SET_AND_BINDING_MAPPING_EXT,
+			.pNext = nullptr,
+			.descriptorSet = VulkanBindless::BindlessSamplerSet,
+			.firstBinding = 0,
+			.bindingCount = 1,
+			.resourceMask = VK_SPIRV_RESOURCE_TYPE_SAMPLER_BIT_EXT,
+			.source = VK_DESCRIPTOR_MAPPING_SOURCE_HEAP_WITH_CONSTANT_OFFSET_EXT,
+			.sourceData = {
+				.constantOffset = {
+					.heapOffset = 0,
+					.heapArrayStride = ResourceState.DescriptorSize,
+					.pEmbeddedSampler = nullptr,
+					.samplerHeapOffset = 0,
+					.samplerHeapArrayStride = SamplerState.DescriptorSize
+				}
+			}
+		};
+
+		const VkSpirvResourceTypeFlagsEXT SpirvResourceTypeFlags = VK_SPIRV_RESOURCE_TYPE_SAMPLED_IMAGE_BIT_EXT | VK_SPIRV_RESOURCE_TYPE_READ_ONLY_IMAGE_BIT_EXT | VK_SPIRV_RESOURCE_TYPE_READ_WRITE_IMAGE_BIT_EXT |
+			VK_SPIRV_RESOURCE_TYPE_UNIFORM_BUFFER_BIT_EXT | VK_SPIRV_RESOURCE_TYPE_READ_ONLY_STORAGE_BUFFER_BIT_EXT | VK_SPIRV_RESOURCE_TYPE_READ_WRITE_STORAGE_BUFFER_BIT_EXT |
+			(bHasRaytracingExtensions ? VK_SPIRV_RESOURCE_TYPE_ACCELERATION_STRUCTURE_BIT_EXT : 0);
+
+		const VkDescriptorSetAndBindingMappingEXT ResourceHeapMapping = {
+			.sType = VK_STRUCTURE_TYPE_DESCRIPTOR_SET_AND_BINDING_MAPPING_EXT,
+			.pNext = nullptr,
+			.descriptorSet = VulkanBindless::BindlessResourceSet,
+			.firstBinding = 0,
+			.bindingCount = 1,
+			.resourceMask = SpirvResourceTypeFlags,
+			.source = VK_DESCRIPTOR_MAPPING_SOURCE_HEAP_WITH_CONSTANT_OFFSET_EXT,
+			.sourceData = {
+				.constantOffset = {
+					.heapOffset = 0,
+					.heapArrayStride = ResourceState.DescriptorSize,
+					.pEmbeddedSampler = nullptr,
+					.samplerHeapOffset = 0,
+					.samplerHeapArrayStride = SamplerState.DescriptorSize
+				}
+			}
+		};
+
+		// We need to offset the push constants to leave place for the static UBs
+		const uint32 PushConstantOffset = (FRHIShaderBindingLayout::MaxUniformBufferEntries * sizeof(uint32));
 
+		for (int32 StageIndex = 0; StageIndex < SF_NumGraphicsFrequencies; ++StageIndex)
+		{
+			const int32 BindingOffset = VulkanBindless::GetOffsetForFrequency((EShaderFrequency)StageIndex);
+			TArray<VkDescriptorSetAndBindingMappingEXT>& StageBindingMappings = GraphicsBindingMappings[StageIndex];
+
+			// Until our DXC supports SamplerHeapEXT / ResourceHeapEXT builtins, add bindings for sampler and resource heaps
+			StageBindingMappings.SetNum(VulkanBindless::MaxUniformBuffersPerStage + 2);
+			StageBindingMappings[VulkanBindless::MaxUniformBuffersPerStage + 0] = SamplerHeapMapping;
+			StageBindingMappings[VulkanBindless::MaxUniformBuffersPerStage + 1] = ResourceHeapMapping;
+
+			for (uint32 UBIndex = 0; UBIndex < (uint32)VulkanBindless::MaxUniformBuffersPerStage; ++UBIndex)
+			{
+				const uint32 BindingIndex = (BindingOffset < 0) ? (VulkanBindless::MaxUniformBuffersTotal - UBIndex - 1) : (BindingOffset + UBIndex);
+
+				StageBindingMappings[UBIndex] = {
+					.sType = VK_STRUCTURE_TYPE_DESCRIPTOR_SET_AND_BINDING_MAPPING_EXT,
+					.pNext = nullptr,
+					.descriptorSet = VulkanBindless::BindlessSingleUseUniformBufferSet,
+					.firstBinding = BindingIndex,
+					.bindingCount = 1,
+					.resourceMask = VK_SPIRV_RESOURCE_TYPE_UNIFORM_BUFFER_BIT_EXT,
+					.source = VK_DESCRIPTOR_MAPPING_SOURCE_HEAP_WITH_PUSH_INDEX_EXT,
+					.sourceData = {
+						.pushIndex = {
+							.heapOffset = 0,
+							.pushOffset = PushConstantOffset + (BindingIndex * 4u),  // as per spec, 4 byte offsets
+							.heapIndexStride = ResourceState.DescriptorSize,
+							.heapArrayStride = ResourceState.DescriptorSize,
+							.pEmbeddedSampler = nullptr,
+							.useCombinedImageSamplerIndex = false,
+							.samplerHeapOffset = 0,
+							.samplerPushOffset = 0,
+							.samplerHeapIndexStride = 0,
+							.samplerHeapArrayStride = 0
+						}
+					}
+				};
+			}
+		}
+
+		// Until our DXC supports SamplerHeapEXT / ResourceHeapEXT builtins, add bindings for sampler and resource heaps
+		ComputeBindingMappings.SetNum(VulkanBindless::MaxUniformBuffersPerStage + 2);
+		ComputeBindingMappings[VulkanBindless::MaxUniformBuffersPerStage + 0] = SamplerHeapMapping;
+		ComputeBindingMappings[VulkanBindless::MaxUniformBuffersPerStage + 1] = ResourceHeapMapping;
+		for (uint32 UBIndex = 0; UBIndex < (uint32)VulkanBindless::MaxUniformBuffersPerStage; ++UBIndex)
+		{
+			ComputeBindingMappings[UBIndex] = {
+				.sType = VK_STRUCTURE_TYPE_DESCRIPTOR_SET_AND_BINDING_MAPPING_EXT,
+				.pNext = nullptr,
+				.descriptorSet = VulkanBindless::BindlessSingleUseUniformBufferSet,
+				.firstBinding = UBIndex,
+				.bindingCount = 1,
+				.resourceMask = VK_SPIRV_RESOURCE_TYPE_UNIFORM_BUFFER_BIT_EXT,
+				.source = VK_DESCRIPTOR_MAPPING_SOURCE_HEAP_WITH_PUSH_INDEX_EXT,
+				.sourceData = {
+					.pushIndex = {
+						.heapOffset = 0,
+						.pushOffset = PushConstantOffset + (UBIndex * 4u),  // as per spec, 4 byte increments
+						.heapIndexStride = ResourceState.DescriptorSize,
+						.heapArrayStride = ResourceState.DescriptorSize,
+						.pEmbeddedSampler = nullptr,
+						.useCombinedImageSamplerIndex = false,
+						.samplerHeapOffset = 0,
+						.samplerPushOffset = 0,
+						.samplerHeapIndexStride = 0,
+						.samplerHeapArrayStride = 0
+					}
+				}
+			};
+		}
+	}
+}
+
+void FVulkanBindlessDescriptorManager::InitDescriptorBuffers()
+{
 	const VkDevice DeviceHandle = Device.GetHandle();
 	const VkPhysicalDeviceDescriptorBufferPropertiesEXT& DescriptorBufferProperties = Device.GetOptionalExtensionProperties().DescriptorBufferProps;
 
@@ -285,28 +583,24 @@ void FVulkanBindlessDescriptorManager::Init()
 		// 0 - Samplers
 		// 1 - All resources types (using mutable)
 		{
-			// Put samplers in state 0 and mutable in 
-			static_assert(UE_ARRAY_COUNT(BindlessSetStates) == 2, "BindlessSetStates array size has changed, must adjust decriptor type distribution!");
-
-			BindlessSetStates[0].DescriptorTypes.Add(VK_DESCRIPTOR_TYPE_SAMPLER);
-
-			BindlessSetStates[1].DescriptorTypes.Add(VK_DESCRIPTOR_TYPE_SAMPLED_IMAGE);
-			BindlessSetStates[1].DescriptorTypes.Add(VK_DESCRIPTOR_TYPE_STORAGE_IMAGE);
-			BindlessSetStates[1].DescriptorTypes.Add(VK_DESCRIPTOR_TYPE_UNIFORM_TEXEL_BUFFER);
-			BindlessSetStates[1].DescriptorTypes.Add(VK_DESCRIPTOR_TYPE_STORAGE_TEXEL_BUFFER);
-			BindlessSetStates[1].DescriptorTypes.Add(VK_DESCRIPTOR_TYPE_STORAGE_BUFFER);
-			BindlessSetStates[1].DescriptorTypes.Add(VK_DESCRIPTOR_TYPE_UNIFORM_BUFFER);
+			SamplerState.DescriptorTypes.Add(VK_DESCRIPTOR_TYPE_SAMPLER);
+
+			ResourceState.DescriptorTypes.Add(VK_DESCRIPTOR_TYPE_SAMPLED_IMAGE);
+			ResourceState.DescriptorTypes.Add(VK_DESCRIPTOR_TYPE_STORAGE_IMAGE);
+			ResourceState.DescriptorTypes.Add(VK_DESCRIPTOR_TYPE_UNIFORM_TEXEL_BUFFER);
+			ResourceState.DescriptorTypes.Add(VK_DESCRIPTOR_TYPE_STORAGE_TEXEL_BUFFER);
+			ResourceState.DescriptorTypes.Add(VK_DESCRIPTOR_TYPE_STORAGE_BUFFER);
+			ResourceState.DescriptorTypes.Add(VK_DESCRIPTOR_TYPE_UNIFORM_BUFFER);
 			if (Device.GetOptionalExtensions().HasRaytracingExtensions())
-	{
-				BindlessSetStates[1].DescriptorTypes.Add(VK_DESCRIPTOR_TYPE_ACCELERATION_STRUCTURE_KHR);
+			{
+				ResourceState.DescriptorTypes.Add(VK_DESCRIPTOR_TYPE_ACCELERATION_STRUCTURE_KHR);
 			}
 
 			// Add at the end so it can be excluded using Num-1 when only actual types are desired
-			BindlessSetStates[1].DescriptorTypes.Add(VK_DESCRIPTOR_TYPE_MUTABLE_EXT);
+			ResourceState.DescriptorTypes.Add(VK_DESCRIPTOR_TYPE_MUTABLE_EXT);
 		}
 
 		// Setup the sampler states
-		BindlessSetState& SamplerState = GetBindlessState(VK_DESCRIPTOR_TYPE_SAMPLER);
 		{
 			SamplerState.DescriptorSize = GetDescriptorTypeSize(Device, VK_DESCRIPTOR_TYPE_SAMPLER);
 			checkf((SamplerState.DescriptorSize > 0), TEXT("Descriptor Type [VK_DESCRIPTOR_TYPE_SAMPLER] returned an invalid descriptor size!"));
@@ -319,23 +613,22 @@ void FVulkanBindlessDescriptorManager::Init()
 			SamplerState.Allocator = new FRHIHeapDescriptorAllocator(TypeMask, SamplerState.DescriptorCapacity, TConstArrayView<TStatId>(AllocatorStats));
 		}
 
-		// Setup the mutable descriptor states
-		BindlessSetState& MutableState = GetBindlessState(VK_DESCRIPTOR_TYPE_MUTABLE_EXT);
+		// Setup the mutable resource descriptor states
 		{
-			MutableState.DescriptorSize = GetDescriptorTypeSize(Device, VK_DESCRIPTOR_TYPE_MUTABLE_EXT);
-			checkf((MutableState.DescriptorSize > 0), TEXT("Descriptor Type [VK_DESCRIPTOR_TYPE_MUTABLE_EXT] returned an invalid descriptor size!"));
-			MutableState.DescriptorCapacity = GVulkanBindlessMaxResourceDescriptorCount;
-			checkf((MutableState.DescriptorCapacity > 0), TEXT("Descriptor Type [VK_DESCRIPTOR_TYPE_MUTABLE_EXT] returned an invalid descriptor count!"));
+			ResourceState.DescriptorSize = GetDescriptorTypeSize(Device, VK_DESCRIPTOR_TYPE_MUTABLE_EXT);
+			checkf((ResourceState.DescriptorSize > 0), TEXT("Descriptor Type [VK_DESCRIPTOR_TYPE_MUTABLE_EXT] returned an invalid descriptor size!"));
+			ResourceState.DescriptorCapacity = GVulkanBindlessMaxResourceDescriptorCount;
+			checkf((ResourceState.DescriptorCapacity > 0), TEXT("Descriptor Type [VK_DESCRIPTOR_TYPE_MUTABLE_EXT] returned an invalid descriptor count!"));
 
 			ERHIDescriptorTypeMask TypeMask{};
 			TArray<TStatId, TInlineAllocator<1>> AllocatorStats;
 			AllocatorStats.Emplace(GET_STATID(STAT_BindlessResourceDescriptorsAllocated));
-			for (int32 TypeIndex = 0; TypeIndex < MutableState.DescriptorTypes.Num() - 1; ++TypeIndex)
+			for (int32 TypeIndex = 0; TypeIndex < ResourceState.DescriptorTypes.Num() - 1; ++TypeIndex)
 			{
-				const VkDescriptorType DescriptorType = MutableState.DescriptorTypes[TypeIndex];
+				const VkDescriptorType DescriptorType = ResourceState.DescriptorTypes[TypeIndex];
 				EnumAddFlags(TypeMask, GetUEDescriptorTypeMaskFromVkDescriptorType(DescriptorType));
 			}
-			MutableState.Allocator = new FRHIHeapDescriptorAllocator(TypeMask, MutableState.DescriptorCapacity, TConstArrayView<TStatId>(AllocatorStats));
+			ResourceState.Allocator = new FRHIHeapDescriptorAllocator(TypeMask, ResourceState.DescriptorCapacity, TConstArrayView<TStatId>(AllocatorStats));
 		}
 
 
@@ -362,11 +655,11 @@ void FVulkanBindlessDescriptorManager::Init()
 				.pBindingFlags = &DescriptorBindingFlags
 			};
 
-			const uint32 MutableTypeCount = MutableState.DescriptorTypes.Num() - 1;
+			const uint32 MutableTypeCount = ResourceState.DescriptorTypes.Num() - 1;
 			const VkMutableDescriptorTypeListEXT MutableDescriptorTypeList =
 			{
 				.descriptorTypeCount = MutableTypeCount,
-				.pDescriptorTypes = MutableState.DescriptorTypes.GetData()
+				.pDescriptorTypes = ResourceState.DescriptorTypes.GetData()
 			};
 
 			const VkMutableDescriptorTypeCreateInfoEXT MutableDescriptorTypeCreateInfo =
@@ -379,7 +672,7 @@ void FVulkanBindlessDescriptorManager::Init()
 
 			// NOTE: Eventually add VkDescriptorSetLayoutBindingFlagsCreateInfo to pNext for VK_DESCRIPTOR_BINDING_VARIABLE_DESCRIPTOR_COUNT_BIT
 			const VkDescriptorSetLayoutCreateInfo DescriptorSetLayoutCreateInfo =
-				{
+			{
 				.sType = VK_STRUCTURE_TYPE_DESCRIPTOR_SET_LAYOUT_CREATE_INFO,
 				.pNext = IsSamplerSet ? (void*)&DescriptorSetLayoutBindingFlagsCreateInfo : (void*)&MutableDescriptorTypeCreateInfo,
 				.flags = VK_DESCRIPTOR_SET_LAYOUT_CREATE_DESCRIPTOR_BUFFER_BIT_EXT,
@@ -387,16 +680,16 @@ void FVulkanBindlessDescriptorManager::Init()
 				.pBindings = &DescriptorSetLayoutBindings
 			};
 
-				VkDescriptorSetLayout DescriptorSetLayout = VK_NULL_HANDLE;
-				VERIFYVULKANRESULT(VulkanRHI::vkCreateDescriptorSetLayout(DeviceHandle, &DescriptorSetLayoutCreateInfo, VULKAN_CPU_ALLOCATOR, &DescriptorSetLayout));
+			VkDescriptorSetLayout DescriptorSetLayout = VK_NULL_HANDLE;
+			VERIFYVULKANRESULT(VulkanRHI::vkCreateDescriptorSetLayout(DeviceHandle, &DescriptorSetLayoutCreateInfo, VULKAN_CPU_ALLOCATOR, &DescriptorSetLayout));
 			VULKAN_SET_DEBUG_NAME(Device, VK_OBJECT_TYPE_DESCRIPTOR_SET_LAYOUT, DescriptorSetLayout, TEXT("BindlessDescriptorSet(%s)"), IsSamplerSet ? TEXT("Sampler") : TEXT("Resources"));
 
 			// Safe check
-				VkDeviceSize BindingOffset = 0;
-				VulkanRHI::vkGetDescriptorSetLayoutBindingOffsetEXT(DeviceHandle, DescriptorSetLayout, 0u, &BindingOffset);
-				check(BindingOffset == 0);
+			VkDeviceSize BindingOffset = 0;
+			VulkanRHI::vkGetDescriptorSetLayoutBindingOffsetEXT(DeviceHandle, DescriptorSetLayout, 0u, &BindingOffset);
+			check(BindingOffset == 0);
 
-				return DescriptorSetLayout;
+			return DescriptorSetLayout;
 		};
 
 
@@ -478,12 +771,10 @@ void FVulkanBindlessDescriptorManager::Init()
 		};
 
 		// Fill in one state for each descriptor type
-		for (uint32 SetIndex = 0; SetIndex < VulkanBindless::NumBindlessSets; ++SetIndex)
-		{
-			BindlessSetState& State = BindlessSetStates[SetIndex];
-			State.DescriptorSetLayout = CreateDescriptorSetLayout(State);
-			CreateDescriptorBuffer(State, BufferBindingInfo[SetIndex]);
-		}
+		SamplerState.DescriptorSetLayout = CreateDescriptorSetLayout(SamplerState);
+		CreateDescriptorBuffer(SamplerState, BufferBindingInfo[VulkanBindless::BindlessSamplerSet]);
+		ResourceState.DescriptorSetLayout = CreateDescriptorSetLayout(ResourceState);
+		CreateDescriptorBuffer(ResourceState, BufferBindingInfo[VulkanBindless::BindlessResourceSet]);
 
 		// Fill in the state for single-use UB
 		// Uniform buffer descriptor set layout differ from the other resources, we reserve a fixed number of descriptors per stage for each draw/dispatch
@@ -516,11 +807,8 @@ void FVulkanBindlessDescriptorManager::Init()
 
 	// Now create the single pipeline layout used by everything
 	{
-		for (int32 LayoutIndex = 0; LayoutIndex < VulkanBindless::NumBindlessSets; ++LayoutIndex)
-		{
-			const BindlessSetState& State = BindlessSetStates[LayoutIndex];
-			BindlessDescriptorSetLayouts[LayoutIndex] = State.DescriptorSetLayout;
-		}
+		BindlessDescriptorSetLayouts[VulkanBindless::BindlessSamplerSet] = SamplerState.DescriptorSetLayout;
+		BindlessDescriptorSetLayouts[VulkanBindless::BindlessResourceSet] = ResourceState.DescriptorSetLayout;
 		BindlessDescriptorSetLayouts[VulkanBindless::BindlessSingleUseUniformBufferSet] = SingleUseUBDescriptorSetLayout;
 		check(VulkanBindless::BindlessSingleUseUniformBufferSet < BindlessDescriptorSetLayouts.Num());
 
@@ -554,32 +842,29 @@ void FVulkanBindlessDescriptorManager::Deinit()
 		BindlessPipelineLayout = VK_NULL_HANDLE;
 
 		auto DestroyBindlessState = [DeviceHandle](BindlessSetState& State)
+		{
+			if (State.DescriptorSetLayout)
 			{
 				VulkanRHI::vkDestroyDescriptorSetLayout(DeviceHandle, State.DescriptorSetLayout, VULKAN_CPU_ALLOCATOR);
 				State.DescriptorSetLayout = VK_NULL_HANDLE;
+			}
 
-				VulkanRHI::vkDestroyBuffer(DeviceHandle, State.BufferHandle, VULKAN_CPU_ALLOCATOR);
-				State.BufferHandle = VK_NULL_HANDLE;
+			VulkanRHI::vkDestroyBuffer(DeviceHandle, State.BufferHandle, VULKAN_CPU_ALLOCATOR);
+			State.BufferHandle = VK_NULL_HANDLE;
 
-				VulkanRHI::vkUnmapMemory(DeviceHandle, State.MemoryHandle);
-				VulkanRHI::vkFreeMemory(DeviceHandle, State.MemoryHandle, VULKAN_CPU_ALLOCATOR);
-				State.MemoryHandle = VK_NULL_HANDLE;
+			VulkanRHI::vkUnmapMemory(DeviceHandle, State.MemoryHandle);
+			VulkanRHI::vkFreeMemory(DeviceHandle, State.MemoryHandle, VULKAN_CPU_ALLOCATOR);
+			State.MemoryHandle = VK_NULL_HANDLE;
 
-				if (State.Allocator)
-				{
-					delete State.Allocator;
-					State.Allocator = nullptr;
-				}
-			};
-
-		for (uint32 SetIndex = 0; SetIndex < VulkanBindless::NumBindlessSets; ++SetIndex)
-		{
-			BindlessSetState& State = BindlessSetStates[SetIndex];
-			if (!State.DescriptorTypes.IsEmpty())
+			if (State.Allocator)
 			{
-				DestroyBindlessState(State);
+				delete State.Allocator;
+				State.Allocator = nullptr;
 			}
-		}
+		};
+
+		DestroyBindlessState(SamplerState);
+		DestroyBindlessState(ResourceState);
 
 		VulkanRHI::vkDestroyDescriptorSetLayout(DeviceHandle, SingleUseUBDescriptorSetLayout, VULKAN_CPU_ALLOCATOR);
 		SingleUseUBDescriptorSetLayout = VK_NULL_HANDLE;
@@ -589,6 +874,7 @@ void FVulkanBindlessDescriptorManager::Deinit()
 void FVulkanBindlessDescriptorManager::BindDescriptorBuffers(VkCommandBuffer CommandBuffer, VkPipelineStageFlags SupportedStages)
 {
 	checkf(IsSupported(), TEXT("Trying to BindDescriptorBuffers but bindless is not supported!"));
+	check(UseDescriptorBuffers());
 
 	VulkanRHI::vkCmdBindDescriptorBuffersEXT(CommandBuffer, VulkanBindless::NumBindlessSets, BufferBindingInfo);
 
@@ -612,8 +898,9 @@ void FVulkanBindlessDescriptorManager::RegisterUniformBuffers(FVulkanCommandList
 {
 	checkf(IsSupported(), TEXT("Trying to RegisterUniformBuffers but bindless is not supported!"));
 	check(UBDescriptors.Num() < VulkanBindless::MaxUniformBuffersTotal);
+	check(UseDescriptorBuffers());
 
-	VkPushDescriptorSetInfo PushDescriptorSetInfo =
+	const VkPushDescriptorSetInfo PushDescriptorSetInfo =
 	{
 		.sType = VK_STRUCTURE_TYPE_PUSH_DESCRIPTOR_SET_INFO,
 		.pNext = nullptr,
@@ -649,43 +936,6 @@ TOptional<FRHIDescriptorAllocation> FVulkanBindlessDescriptorManager::AllocateDe
 	return TOptional<FRHIDescriptorAllocation>();
 }
 
-
-FRHIDescriptorHandle FVulkanBindlessDescriptorManager::AllocateDescriptor(ERHIDescriptorType DescriptorType, VkDeviceAddress BufferAddress, VkDeviceSize BufferSize)
-{
-	if (IsSupported())
-	{
-		FRHIDescriptorHandle DescriptorHandle = AllocateDescriptor(DescriptorType);
-		UpdateDescriptor({}, DescriptorHandle, BufferAddress, BufferSize, true);
-		return DescriptorHandle;
-	}
-
-	return FRHIDescriptorHandle();
-}
-
-FRHIDescriptorHandle FVulkanBindlessDescriptorManager::AllocateDescriptor(ERHIDescriptorType DescriptorType, FVulkanBuffer* Buffer, uint32 ExtraOffset)
-{
-	if (IsSupported())
-	{
-		FRHIDescriptorHandle DescriptorHandle = AllocateDescriptor(DescriptorType);
-		UpdateDescriptor({}, DescriptorHandle, Buffer->GetHandle(), Buffer->GetOffset() + ExtraOffset, Buffer->GetCurrentSize() - ExtraOffset, true);
-		return DescriptorHandle;
-	}
-
-	return FRHIDescriptorHandle();
-}
-
-FRHIDescriptorHandle FVulkanBindlessDescriptorManager::AllocateDescriptor(VkSampler VulkanSampler)
-{
-	if (IsSupported())
-	{
-		FRHIDescriptorHandle DescriptorHandle = AllocateDescriptor(ERHIDescriptorType::Sampler);
-		UpdateDescriptor({}, DescriptorHandle, VulkanSampler);
-		return DescriptorHandle;
-	}
-
-	return FRHIDescriptorHandle();
-}
-
 void FVulkanBindlessDescriptorManager::FreeDescriptor(FRHIDescriptorHandle DescriptorHandle)
 {
 	if (DescriptorHandle.IsValid())
@@ -695,8 +945,11 @@ void FVulkanBindlessDescriptorManager::FreeDescriptor(FRHIDescriptorHandle Descr
 		BindlessSetState& State = GetBindlessState(DescriptorHandle.GetType());
 		State.Allocator->Free(DescriptorHandle);
 
-		const uint32 ByteOffset = DescriptorHandle.GetIndex() * State.DescriptorSize;
-		FMemory::Memzero(&State.CPUDescriptorMemory[ByteOffset], State.DescriptorSize); // easier for debugging for now
+		if (UseDescriptorBuffers())
+		{
+			const uint32 ByteOffset = DescriptorHandle.GetIndex() * State.DescriptorSize;
+			FMemory::Memzero(&State.CPUDescriptorMemory[ByteOffset], State.DescriptorSize); // easier for debugging for now
+		}
 	}
 }
 
@@ -709,13 +962,17 @@ void FVulkanBindlessDescriptorManager::FreeDescriptors(FRHIDescriptorAllocation
 		BindlessSetState& State = GetBindlessState(DescriptorType);
 		State.Allocator->Free(Allocation);
 
-		const uint32 ByteOffset = Allocation.StartIndex * State.DescriptorSize;
-		FMemory::Memzero(&State.CPUDescriptorMemory[ByteOffset], State.DescriptorSize * Allocation.Count); // easier for debugging for now
+		if (UseDescriptorBuffers())
+		{
+			const uint32 ByteOffset = Allocation.StartIndex * State.DescriptorSize;
+			FMemory::Memzero(&State.CPUDescriptorMemory[ByteOffset], State.DescriptorSize * Allocation.Count); // easier for debugging for now
+		}
 	}
 }
 
 void FVulkanBindlessDescriptorManager::UpdateDescriptorInternal(const FVulkanContextArray& Contexts, FRHIDescriptorHandle DescriptorHandle, VkDescriptorDataEXT DescriptorData, bool bImmediateUpdate)
 {
+	check(UseDescriptorBuffers());
 	checkf(DescriptorHandle.IsValid(), TEXT("Attemping to update invalid descriptor handle!"));
 
 	const VkDescriptorType DescriptorType = GetVkDescriptorTypeFromUEDescriptorType(DescriptorHandle.GetType());
@@ -783,13 +1040,14 @@ void FVulkanBindlessDescriptorManager::UpdateDescriptorInternal(const FVulkanCon
 	}
 }
 
-void FVulkanBindlessDescriptorManager::UpdateDescriptor(const FVulkanContextArray& Contexts, FRHIDescriptorHandle DescriptorHandle, VkSampler VulkanSampler)
+void FVulkanBindlessDescriptorManager::UpdateDescriptor(FRHIDescriptorHandle DescriptorHandle, VkSampler VulkanSampler)
 {
 	if (IsSupported())
 	{
+		check(UseDescriptorBuffers());
 		VkDescriptorDataEXT DescriptorData;
 		DescriptorData.pSampler = &VulkanSampler;
-		UpdateDescriptorInternal(Contexts, DescriptorHandle, DescriptorData, true);
+		UpdateDescriptorInternal({}, DescriptorHandle, DescriptorData, true);
 	}
 }
 
@@ -797,6 +1055,7 @@ void FVulkanBindlessDescriptorManager::UpdateDescriptor(const FVulkanContextArra
 {
 	if (IsSupported())
 	{
+		check(UseDescriptorBuffers());
 		check(DescriptorHandle.GetType() == ERHIDescriptorType::TextureSRV || DescriptorHandle.GetType() == ERHIDescriptorType::TextureUAV);
 
 		VkDescriptorImageInfo DescriptorImageInfo;
@@ -817,14 +1076,31 @@ void FVulkanBindlessDescriptorManager::UpdateDescriptor(const FVulkanContextArra
 {
 	if (IsSupported())
 	{
-		check(DescriptorHandle.GetType() == ERHIDescriptorType::BufferSRV || DescriptorHandle.GetType() == ERHIDescriptorType::BufferUAV);
-
-		VkBufferDeviceAddressInfo BufferInfo;
-		ZeroVulkanStruct(BufferInfo, VK_STRUCTURE_TYPE_BUFFER_DEVICE_ADDRESS_INFO);
-		BufferInfo.buffer = VulkanBuffer;
+		const VkBufferDeviceAddressInfo BufferInfo =
+		{
+			.sType = VK_STRUCTURE_TYPE_BUFFER_DEVICE_ADDRESS_INFO,
+			.pNext = nullptr,
+			.buffer = VulkanBuffer
+		};
 		const VkDeviceAddress BufferAddress = VulkanRHI::vkGetBufferDeviceAddressKHR(Device.GetHandle(), &BufferInfo);
 
-		UpdateDescriptor(Contexts, DescriptorHandle, BufferAddress + BufferOffset, BufferSize, bImmediateUpdate);
+		if (UseDescriptorBuffers())
+		{
+			check(DescriptorHandle.GetType() == ERHIDescriptorType::CBV || DescriptorHandle.GetType() == ERHIDescriptorType::BufferSRV || DescriptorHandle.GetType() == ERHIDescriptorType::BufferUAV);
+			UpdateDescriptor(Contexts, DescriptorHandle, BufferAddress + BufferOffset, BufferSize, bImmediateUpdate);
+		}
+		else
+		{
+			check(DescriptorHandle.GetType() == ERHIDescriptorType::CBV || DescriptorHandle.GetType() == ERHIDescriptorType::BufferSRV || DescriptorHandle.GetType() == ERHIDescriptorType::BufferUAV || DescriptorHandle.GetType() == ERHIDescriptorType::AccelerationStructure);
+			check(UseDescriptorHeaps());
+
+			const VkDeviceAddressRangeEXT DeviceAddressRange = {
+				.address = BufferAddress + BufferOffset,
+				.size = BufferSize
+			};
+
+			UpdateDescriptor(Contexts, DescriptorHandle, DeviceAddressRange, bImmediateUpdate);
+		}
 	}
 }
 
@@ -834,14 +1110,28 @@ void FVulkanBindlessDescriptorManager::UpdateDescriptor(const FVulkanContextArra
 	{
 		check(DescriptorHandle.GetType() == ERHIDescriptorType::BufferSRV || DescriptorHandle.GetType() == ERHIDescriptorType::CBV || DescriptorHandle.GetType() == ERHIDescriptorType::BufferUAV);
 
-		VkDescriptorAddressInfoEXT AddressInfo;
-		ZeroVulkanStruct(AddressInfo, VK_STRUCTURE_TYPE_DESCRIPTOR_ADDRESS_INFO_EXT);
-		AddressInfo.address = BufferAddress;
-		AddressInfo.range = BufferSize;
+		if (UseDescriptorBuffers())
+		{
+			VkDescriptorAddressInfoEXT AddressInfo;
+			ZeroVulkanStruct(AddressInfo, VK_STRUCTURE_TYPE_DESCRIPTOR_ADDRESS_INFO_EXT);
+			AddressInfo.address = BufferAddress;
+			AddressInfo.range = BufferSize;
+
+			VkDescriptorDataEXT DescriptorData;
+			DescriptorData.pStorageBuffer = &AddressInfo;  // same pointer for uniform, it's a union
+			UpdateDescriptorInternal(Contexts, DescriptorHandle, DescriptorData, bImmediateUpdate);
+		}
+		else
+		{
+			check(UseDescriptorHeaps());
 
-		VkDescriptorDataEXT DescriptorData;
-		DescriptorData.pStorageBuffer = &AddressInfo;  // same pointer for uniform, it's a union
-		UpdateDescriptorInternal(Contexts, DescriptorHandle, DescriptorData, bImmediateUpdate);
+			const VkDeviceAddressRangeEXT DeviceAddressRange = {
+				.address = BufferAddress,
+				.size = BufferSize
+			};
+
+			UpdateDescriptor(Contexts, DescriptorHandle, DeviceAddressRange, bImmediateUpdate);
+		}
 	}
 }
 
@@ -852,20 +1142,41 @@ void FVulkanBindlessDescriptorManager::UpdateDescriptor(const FVulkanContextArra
 		check(DescriptorHandle.GetType() == ERHIDescriptorType::TypedBufferSRV || DescriptorHandle.GetType() == ERHIDescriptorType::TypedBufferUAV);
 
 		// :todo-jn: start caching buffer addresses in resources to avoid the extra call
-		VkBufferDeviceAddressInfo BufferInfo;
-		ZeroVulkanStruct(BufferInfo, VK_STRUCTURE_TYPE_BUFFER_DEVICE_ADDRESS_INFO);
-		BufferInfo.buffer = ViewInfo.buffer;
+		const VkBufferDeviceAddressInfo BufferInfo = {
+			.sType = VK_STRUCTURE_TYPE_BUFFER_DEVICE_ADDRESS_INFO,
+			.pNext = nullptr,
+			.buffer = ViewInfo.buffer
+		};
 		const VkDeviceAddress BufferAddress = VulkanRHI::vkGetBufferDeviceAddressKHR(Device.GetHandle(), &BufferInfo);
 
-		VkDescriptorAddressInfoEXT AddressInfo;
-		ZeroVulkanStruct(AddressInfo, VK_STRUCTURE_TYPE_DESCRIPTOR_ADDRESS_INFO_EXT);
-		AddressInfo.address = BufferAddress + ViewInfo.offset;
-		AddressInfo.range = ViewInfo.range;
-		AddressInfo.format = ViewInfo.format;
+		if (UseDescriptorBuffers())
+		{
+			VkDescriptorAddressInfoEXT AddressInfo;
+			ZeroVulkanStruct(AddressInfo, VK_STRUCTURE_TYPE_DESCRIPTOR_ADDRESS_INFO_EXT);
+			AddressInfo.address = BufferAddress + ViewInfo.offset;
+			AddressInfo.range = ViewInfo.range;
+			AddressInfo.format = ViewInfo.format;
+
+			VkDescriptorDataEXT DescriptorData;
+			DescriptorData.pUniformTexelBuffer = &AddressInfo;  // same pointer for storage, it's a union
+			UpdateDescriptorInternal(Contexts, DescriptorHandle, DescriptorData, bImmediateUpdate);
+		}
+		else
+		{
+			check(UseDescriptorHeaps());
 
-		VkDescriptorDataEXT DescriptorData;
-		DescriptorData.pUniformTexelBuffer = &AddressInfo;  // same pointer for storage, it's a union
-		UpdateDescriptorInternal(Contexts, DescriptorHandle, DescriptorData, bImmediateUpdate);
+			const VkTexelBufferDescriptorInfoEXT TexelBufferDescriptorInfo = {
+				.sType = VK_STRUCTURE_TYPE_TEXEL_BUFFER_DESCRIPTOR_INFO_EXT,
+				.pNext = nullptr,
+				.format = ViewInfo.format,
+				.addressRange = {
+					.address = BufferAddress + ViewInfo.offset,
+					.size = ViewInfo.range
+				}
+			};
+
+			UpdateDescriptor(Contexts, DescriptorHandle, TexelBufferDescriptorInfo, bImmediateUpdate);
+		}
 	}
 }
 
@@ -873,6 +1184,8 @@ void FVulkanBindlessDescriptorManager::UpdateDescriptor(const FVulkanContextArra
 {
 	if (IsSupported())
 	{
+		check(UseDescriptorBuffers());
+
 		// :todo-jn: start caching AccelerationStructure in resources to avoid the extra call
 		VkAccelerationStructureDeviceAddressInfoKHR AccelerationStructureDeviceAddressInfo;
 		ZeroVulkanStruct(AccelerationStructureDeviceAddressInfo, VK_STRUCTURE_TYPE_ACCELERATION_STRUCTURE_DEVICE_ADDRESS_INFO_KHR);
@@ -884,3 +1197,112 @@ void FVulkanBindlessDescriptorManager::UpdateDescriptor(const FVulkanContextArra
 		UpdateDescriptorInternal(Contexts, DescriptorHandle, DescriptorData, bImmediateUpdate);
 	}
 }
+
+
+void FVulkanBindlessDescriptorManager::UpdateDescriptor(FRHIDescriptorHandle DescriptorHandle, const VkSamplerCreateInfo& SamplerCreateInfo)
+{
+	if (IsSupported())
+	{
+		check(UseDescriptorHeaps());
+
+		const uint32 ByteOffset = DescriptorHandle.GetIndex() * SamplerState.DescriptorSize;
+
+		const VkHostAddressRangeEXT HostAddressRangeEXT = {
+			.address = &SamplerState.MappedPointer[ByteOffset],
+			.size = SamplerState.DescriptorSize
+		};
+
+		VERIFYVULKANRESULT(VulkanRHI::vkWriteSamplerDescriptorsEXT(Device.GetHandle(), 1 /*samplerCount*/, &SamplerCreateInfo, &HostAddressRangeEXT));
+	}
+}
+
+void FVulkanBindlessDescriptorManager::UpdateDescriptorInternal(const FVulkanContextArray& Contexts, FRHIDescriptorHandle DescriptorHandle, const VkResourceDescriptorDataEXT& DescriptorData, bool bImmediateUpdate)
+{
+	if (IsSupported())
+	{
+		check(UseDescriptorHeaps());
+
+		const VkResourceDescriptorInfoEXT ResourceDescriptorInfo = {
+			.sType = VK_STRUCTURE_TYPE_RESOURCE_DESCRIPTOR_INFO_EXT,
+			.pNext = nullptr,
+			.type = GetVkDescriptorTypeFromUEDescriptorType(DescriptorHandle.GetType()),
+			.data = DescriptorData
+		};
+
+		const uint32 ByteOffset = DescriptorHandle.GetIndex() * ResourceState.DescriptorSize;
+
+		// :todo-jn:  Hack to avoid barriers/copies in renderpasses
+		FVulkanCommandListContext* Context = Contexts[ERHIPipeline::Graphics];
+		const bool bIsInsideRenderPass = Context && Context->GetCommandBuffer().IsInsideRenderPass();
+
+		if (bImmediateUpdate || bIsInsideRenderPass)
+		{
+			const VkHostAddressRangeEXT HostAddressRangeEXT = {
+				.address = &ResourceState.MappedPointer[ByteOffset],
+				.size = ResourceState.DescriptorSize
+			};
+
+			VERIFYVULKANRESULT(VulkanRHI::vkWriteResourceDescriptorsEXT(Device.GetHandle(), 1 /*resourceCount*/, &ResourceDescriptorInfo, &HostAddressRangeEXT));
+		}
+		else
+		{
+			checkf(Context, TEXT("Context was NULL, could not update descriptor"));
+
+			VulkanRHI::FStagingBuffer* StagingBuffer = Device.GetStagingManager().AcquireBuffer(ResourceState.DescriptorSize);
+			const VkHostAddressRangeEXT HostAddressRangeEXT = {
+				.address = StagingBuffer->GetMappedPointer(),
+				.size = ResourceState.DescriptorSize
+			};
+			VERIFYVULKANRESULT(VulkanRHI::vkWriteResourceDescriptorsEXT(Device.GetHandle(), 1 /*resourceCount*/, &ResourceDescriptorInfo, &HostAddressRangeEXT));
+
+			VkMemoryBarrier2 MemoryBarrier = {
+				.sType = VK_STRUCTURE_TYPE_MEMORY_BARRIER_2,
+				.pNext = nullptr,
+				.srcStageMask = VK_PIPELINE_STAGE_2_ALL_COMMANDS_BIT,
+				.srcAccessMask = VK_ACCESS_2_RESOURCE_HEAP_READ_BIT_EXT | VK_ACCESS_2_MEMORY_READ_BIT | VK_ACCESS_2_MEMORY_WRITE_BIT,
+				.dstStageMask = VK_PIPELINE_STAGE_2_TRANSFER_BIT,
+				.dstAccessMask = VK_ACCESS_2_TRANSFER_READ_BIT | VK_ACCESS_2_TRANSFER_WRITE_BIT
+			};
+
+			FVulkanCommandBuffer& CommandBuffer = Context->GetCommandBuffer();
+
+			VkDependencyInfo DependencyInfo;
+			ZeroVulkanStruct(DependencyInfo, VK_STRUCTURE_TYPE_DEPENDENCY_INFO);
+			DependencyInfo.memoryBarrierCount = 1;
+			DependencyInfo.pMemoryBarriers = &MemoryBarrier;
+			VulkanRHI::vkCmdPipelineBarrier2KHR(CommandBuffer.GetHandle(), &DependencyInfo);
+
+			const VkBufferCopy Region = {
+				.srcOffset = 0,
+				.dstOffset = ByteOffset,
+				.size = ResourceState.DescriptorSize
+			};
+			VulkanRHI::vkCmdCopyBuffer(CommandBuffer.GetHandle(), StagingBuffer->GetHandle(), ResourceState.BufferHandle, 1, &Region);
+
+			MemoryBarrier.srcStageMask = VK_PIPELINE_STAGE_2_TRANSFER_BIT;
+			MemoryBarrier.srcAccessMask = VK_ACCESS_2_TRANSFER_READ_BIT | VK_ACCESS_2_TRANSFER_WRITE_BIT;
+			MemoryBarrier.dstStageMask = VK_PIPELINE_STAGE_2_ALL_COMMANDS_BIT;
+			MemoryBarrier.dstAccessMask = VK_ACCESS_2_RESOURCE_HEAP_READ_BIT_EXT;
+			VulkanRHI::vkCmdPipelineBarrier2KHR(CommandBuffer.GetHandle(), &DependencyInfo);
+
+			Device.GetStagingManager().ReleaseBuffer(Context, StagingBuffer);
+		}
+	}
+}
+
+TConstArrayView<VkDescriptorSetAndBindingMappingEXT> FVulkanBindlessDescriptorManager::GetBindingMappings(EShaderFrequency Frequency) const
+{
+	if (UseDescriptorBuffers())
+	{
+		return {};
+	}
+
+	if (IsValidGraphicsFrequency(Frequency))
+	{
+		return GraphicsBindingMappings[Frequency];
+	}
+	else
+	{
+		return ComputeBindingMappings;
+	}
+}
diff --git a/Engine/Source/Runtime/VulkanRHI/Private/VulkanBindlessDescriptorManager.h b/Engine/Source/Runtime/VulkanRHI/Private/VulkanBindlessDescriptorManager.h
--- a/Engine/Source/Runtime/VulkanRHI/Private/VulkanBindlessDescriptorManager.h
+++ b/Engine/Source/Runtime/VulkanRHI/Private/VulkanBindlessDescriptorManager.h
@@ -16,6 +16,12 @@ class FVulkanBuffer;
 class FVulkanBindlessDescriptorManager
 {
 public:
+	enum class EBindlessType : uint8
+	{
+		DescriptorHeap,
+		DescriptorBuffer,
+		Disabled
+	};
 
 	enum class EPushConstantRangeType : uint8
 	{
@@ -30,30 +36,41 @@ public:
 	void Init();
 	void Deinit();
 
-	void BindDescriptorBuffers(VkCommandBuffer CommandBuffer, VkPipelineStageFlags SupportedStages);
-
 	FRHIDescriptorHandle AllocateDescriptor(ERHIDescriptorType DescriptorType);
 	TOptional<FRHIDescriptorAllocation> AllocateDescriptors(ERHIDescriptorType DescriptorType, uint32 DescriptorCount);
-
-	FRHIDescriptorHandle AllocateDescriptor(ERHIDescriptorType DescriptorType, VkDeviceAddress BufferAddress, VkDeviceSize BufferSize);
-	FRHIDescriptorHandle AllocateDescriptor(ERHIDescriptorType DescriptorType, FVulkanBuffer* Buffer, uint32 ExtraOffset);
-
-	FRHIDescriptorHandle AllocateDescriptor(VkSampler VulkanSampler);
-
 	void FreeDescriptor(FRHIDescriptorHandle DescriptorHandle);
 	void FreeDescriptors(FRHIDescriptorAllocation Allocation, ERHIDescriptorType DescriptorType);
 
-	void UpdateDescriptor(const FVulkanContextArray& Contexts, FRHIDescriptorHandle DescriptorHandle, VkSampler VulkanSampler);
-	void UpdateDescriptor(const FVulkanContextArray& Contexts, FRHIDescriptorHandle DescriptorHandle, VkImageView VulkanImage, bool bIsDepthStencil, bool bImmediateUpdate = true);
+	// Descriptor Heap update functions
+	void UpdateDescriptor(FRHIDescriptorHandle DescriptorHandle, const VkSamplerCreateInfo& SamplerCreateInfo);
+	void UpdateDescriptor(const FVulkanContextArray& Contexts, FRHIDescriptorHandle DescriptorHandle, const VkImageDescriptorInfoEXT& ImageDescriptorInfo, bool bImmediateUpdate = true)
+	{
+		check(DescriptorHandle.GetType() == ERHIDescriptorType::TextureSRV || DescriptorHandle.GetType() == ERHIDescriptorType::TextureUAV);
+		UpdateDescriptorInternal(Contexts, DescriptorHandle, VkResourceDescriptorDataEXT{ .pImage = &ImageDescriptorInfo }, bImmediateUpdate);
+	}
+	void UpdateDescriptor(const FVulkanContextArray& Contexts, FRHIDescriptorHandle DescriptorHandle, const VkTexelBufferDescriptorInfoEXT& TexelBufferDescriptorInfo, bool bImmediateUpdate = true)
+	{
+		check(DescriptorHandle.GetType() == ERHIDescriptorType::TypedBufferSRV || DescriptorHandle.GetType() == ERHIDescriptorType::TypedBufferUAV);
+		UpdateDescriptorInternal(Contexts, DescriptorHandle, VkResourceDescriptorDataEXT{ .pTexelBuffer = &TexelBufferDescriptorInfo }, bImmediateUpdate);
+	}
+	void UpdateDescriptor(const FVulkanContextArray& Contexts, FRHIDescriptorHandle DescriptorHandle, const VkDeviceAddressRangeEXT& DeviceAddressRange, bool bImmediateUpdate = true)
+	{
+		check(DescriptorHandle.GetType() == ERHIDescriptorType::CBV || DescriptorHandle.GetType() == ERHIDescriptorType::BufferSRV ||DescriptorHandle.GetType() == ERHIDescriptorType::BufferUAV || DescriptorHandle.GetType() == ERHIDescriptorType::AccelerationStructure);
+		UpdateDescriptorInternal(Contexts, DescriptorHandle, VkResourceDescriptorDataEXT{ .pAddressRange = &DeviceAddressRange }, bImmediateUpdate);
+	}
 
+	// Descriptor Buffer functions
+	void BindDescriptorBuffers(VkCommandBuffer CommandBuffer, VkPipelineStageFlags SupportedStages);
+	void UpdateDescriptor(FRHIDescriptorHandle DescriptorHandle, VkSampler VulkanSampler);
+	void UpdateDescriptor(const FVulkanContextArray& Contexts, FRHIDescriptorHandle DescriptorHandle, VkImageView VulkanImage, bool bIsDepthStencil, bool bImmediateUpdate = true);
 	void UpdateDescriptor(const FVulkanContextArray& Contexts, FRHIDescriptorHandle DescriptorHandle, VkBuffer VulkanBuffer, VkDeviceSize BufferOffset, VkDeviceSize BufferSize, bool bImmediateUpdate = true);
 	void UpdateDescriptor(const FVulkanContextArray& Contexts, FRHIDescriptorHandle DescriptorHandle, VkDeviceAddress BufferAddress, VkDeviceSize BufferSize, bool bImmediateUpdate = true);
-
 	void UpdateDescriptor(const FVulkanContextArray& Contexts, FRHIDescriptorHandle DescriptorHandle, const VkBufferViewCreateInfo& ViewInfo, bool bImmediateUpdate = true);
 	void UpdateDescriptor(const FVulkanContextArray& Contexts, FRHIDescriptorHandle DescriptorHandle, VkAccelerationStructureKHR AccelerationStructure, bool bImmediateUpdate = true);
-
 	void RegisterUniformBuffers(FVulkanCommandListContext& Context, VkShaderStageFlags StageFlags, TConstArrayView<VkWriteDescriptorSet> UBDescriptors);
 
+	TConstArrayView<VkDescriptorSetAndBindingMappingEXT> GetBindingMappings(EShaderFrequency Frequency) const;
+
 	ERHIBindlessConfiguration GetConfiguration() const
 	{
 		return Configuration;
@@ -66,12 +83,13 @@ public:
 
 	VkPipelineCreateFlagBits2 GetPipelineCreateFlag() const
 	{
-		return VK_PIPELINE_CREATE_2_DESCRIPTOR_BUFFER_BIT_EXT;
+		return UseDescriptorBuffers() ? VK_PIPELINE_CREATE_2_DESCRIPTOR_BUFFER_BIT_EXT :
+			   UseDescriptorHeaps() ? VK_PIPELINE_CREATE_2_DESCRIPTOR_HEAP_BIT_EXT : 0ull;
 	}
 
 	VkPipelineLayout GetPipelineLayout() const
 	{
-		return BindlessPipelineLayout;
+		return UseDescriptorBuffers() ? BindlessPipelineLayout : nullptr;
 	}
 
 	// Returns descriptor set layouts for a single stage
@@ -80,9 +98,29 @@ public:
 		return BindlessDescriptorSetLayouts;
 	}
 
-	static TConstArrayView<VkPushConstantRange> GetPushConstantRanges()
+	TConstArrayView<VkPushConstantRange> GetPushConstantRanges() const
+	{
+		return UseDescriptorHeaps() ? TConstArrayView<VkPushConstantRange>() : PushConstantRanges;
+	}
+
+	bool UseDescriptorHeaps() const
+	{
+		return BindlessType == EBindlessType::DescriptorHeap;
+	}
+
+	bool UseDescriptorBuffers() const
+	{
+		return BindlessType == EBindlessType::DescriptorBuffer;
+	}
+
+	const VkBindHeapInfoEXT& GetSamplerBindHeapInfo() const
+	{
+		return SamplerBindHeapInfo;
+	}
+
+	const VkBindHeapInfoEXT& GetResourceBindHeapInfo() const
 	{
-		return PushConstantRanges;
+		return ResourceBindHeapInfo;
 	}
 
 private:
@@ -91,6 +129,8 @@ private:
 	ERHIBindlessConfiguration Configuration{};
 	ERHIPipeline              ConfiguredPipelines = ERHIPipeline::None;
 
+	EBindlessType BindlessType = EBindlessType::Disabled;
+
 	struct BindlessSetState
 	{
 		~BindlessSetState()
@@ -108,32 +148,44 @@ private:
 			}
 		}
 
-		TArray<VkDescriptorType, TInlineAllocator<2>> DescriptorTypes;
 		FRHIHeapDescriptorAllocator* Allocator = nullptr;
 
-		VkDescriptorSetLayout DescriptorSetLayout = VK_NULL_HANDLE;
-
-		uint32 DescriptorCapacity = 0;
 		uint32 DescriptorSize = 0;
 
 		VkBuffer BufferHandle = VK_NULL_HANDLE;
 		VkDeviceMemory MemoryHandle = VK_NULL_HANDLE;
 		uint8* MappedPointer = nullptr;
 
+		// Descriptor Buffer specific
+		TArray<VkDescriptorType, TInlineAllocator<2>> DescriptorTypes;
+		VkDescriptorSetLayout DescriptorSetLayout = VK_NULL_HANDLE;
+		uint32 DescriptorCapacity = 0;
 		uint8* CPUDescriptorMemory = nullptr;
 	};
-	BindlessSetState BindlessSetStates[VulkanBindless::NumBindlessSets];
+	BindlessSetState SamplerState;
+	BindlessSetState ResourceState;
 
-	VkDescriptorSetLayout SingleUseUBDescriptorSetLayout = VK_NULL_HANDLE;
+	BindlessSetState& GetBindlessState(ERHIDescriptorType RHIDescriptorType);
 
+	// Descriptor Heap specific values
+	VkBindHeapInfoEXT SamplerBindHeapInfo;
+	VkBindHeapInfoEXT ResourceBindHeapInfo;
+	TArray<VkDescriptorSetAndBindingMappingEXT> GraphicsBindingMappings[SF_NumGraphicsFrequencies];
+	TArray<VkDescriptorSetAndBindingMappingEXT> ComputeBindingMappings;
+
+	// Descriptor Heap functions
+	void InitDescriptorHeaps();
+	void UpdateDescriptorInternal(const FVulkanContextArray& Contexts, FRHIDescriptorHandle DescriptorHandle, const VkResourceDescriptorDataEXT& DescriptorData, bool bImmediateUpdate);
+
+	// Descriptor Buffer specific values
+	VkDescriptorSetLayout SingleUseUBDescriptorSetLayout = VK_NULL_HANDLE;
 	VkDescriptorBufferBindingInfoEXT BufferBindingInfo[VulkanBindless::NumBindlessSets];
 	uint32_t BufferIndices[VulkanBindless::MaxNumSets];
 	VkPipelineLayout BindlessPipelineLayout = VK_NULL_HANDLE;
 	TStaticArray<VkDescriptorSetLayout, VulkanBindless::MaxNumSets> BindlessDescriptorSetLayouts;
 	static const VkPushConstantRange PushConstantRanges[(uint8)EPushConstantRangeType::Count];
 
-	BindlessSetState& GetBindlessState(ERHIDescriptorType RHIDescriptorType);
-	BindlessSetState& GetBindlessState(VkDescriptorType DescriptorType);
-
+	// Descriptor Buffer functions
+	void InitDescriptorBuffers();
 	void UpdateDescriptorInternal(const FVulkanContextArray& Contexts, FRHIDescriptorHandle DescriptorHandle, VkDescriptorDataEXT DescriptorData, bool bImmediateUpdate);
 };
diff --git a/Engine/Source/Runtime/VulkanRHI/Private/VulkanCommandBuffer.cpp b/Engine/Source/Runtime/VulkanRHI/Private/VulkanCommandBuffer.cpp
--- a/Engine/Source/Runtime/VulkanRHI/Private/VulkanCommandBuffer.cpp
+++ b/Engine/Source/Runtime/VulkanRHI/Private/VulkanCommandBuffer.cpp
@@ -235,7 +235,8 @@ void FVulkanCommandBuffer::End(FVulkanQueryPool* OptionalQueryPool)
 
 void FVulkanCommandBuffer::Begin(FVulkanQueryPool* OptionalQueryPool, VkRenderPass RenderPassHandle)
 {
-	checkf((GetCommandBufferType() != EVulkanCommandBufferType::Secondary) || (RenderPassHandle != VK_NULL_HANDLE),
+	const bool bIsSecondaryCommandBuffer = GetCommandBufferType() == EVulkanCommandBufferType::Secondary;
+	checkf(!bIsSecondaryCommandBuffer || (RenderPassHandle != VK_NULL_HANDLE),
 		TEXT("Secondary command buffers require the render pass handle!"));
 
 	{
@@ -255,23 +256,44 @@ void FVulkanCommandBuffer::Begin(FVulkanQueryPool* OptionalQueryPool, VkRenderPa
 	ZeroVulkanStruct(BeginInfo, VK_STRUCTURE_TYPE_COMMAND_BUFFER_BEGIN_INFO);
 	BeginInfo.flags = VK_COMMAND_BUFFER_USAGE_ONE_TIME_SUBMIT_BIT;
 
+	FVulkanBindlessDescriptorManager* BindlessDescriptorManager = Device.GetBindlessDescriptorManager();
 	VkCommandBufferInheritanceInfo InheritanceInfo;
-	if (GetCommandBufferType() == EVulkanCommandBufferType::Secondary)
+	VkCommandBufferInheritanceDescriptorHeapInfoEXT DescriptorHeapInfo;
+	if (bIsSecondaryCommandBuffer)
 	{
 		BeginInfo.flags |= VK_COMMAND_BUFFER_USAGE_RENDER_PASS_CONTINUE_BIT;
 		BeginInfo.pInheritanceInfo = &InheritanceInfo;
 
 		ZeroVulkanStruct(InheritanceInfo, VK_STRUCTURE_TYPE_COMMAND_BUFFER_INHERITANCE_INFO);
 		InheritanceInfo.renderPass = RenderPassHandle;
+
+		if (BindlessDescriptorManager->UseDescriptorHeaps())
+		{
+			DescriptorHeapInfo = {
+				.sType = VK_STRUCTURE_TYPE_COMMAND_BUFFER_INHERITANCE_DESCRIPTOR_HEAP_INFO_EXT,
+				.pNext = nullptr,
+				.pSamplerHeapBindInfo = &BindlessDescriptorManager->GetSamplerBindHeapInfo(),
+				.pResourceHeapBindInfo = &BindlessDescriptorManager->GetResourceBindHeapInfo()
+			};
+
+			AddToPNext(InheritanceInfo, DescriptorHeapInfo);
+		}
 	}
 
 	VERIFYVULKANRESULT(VulkanRHI::vkBeginCommandBuffer(CommandBufferHandle, &BeginInfo));
 
 	if (Device.SupportsBindless())
 	{
-		FVulkanBindlessDescriptorManager* BindlessDescriptorManager = Device.GetBindlessDescriptorManager();
-		const VkPipelineStageFlags SupportedStages = CommandBufferPool.GetQueue().GetSupportedStageBits();
-		BindlessDescriptorManager->BindDescriptorBuffers(CommandBufferHandle, SupportedStages);
+		if (BindlessDescriptorManager->UseDescriptorBuffers())
+		{
+			const VkPipelineStageFlags SupportedStages = CommandBufferPool.GetQueue().GetSupportedStageBits();
+			BindlessDescriptorManager->BindDescriptorBuffers(CommandBufferHandle, SupportedStages);
+		}
+		else if (BindlessDescriptorManager->UseDescriptorHeaps() && !bIsSecondaryCommandBuffer)
+		{
+			VulkanRHI::vkCmdBindSamplerHeapEXT(CommandBufferHandle, &BindlessDescriptorManager->GetSamplerBindHeapInfo());
+			VulkanRHI::vkCmdBindResourceHeapEXT(CommandBufferHandle, &BindlessDescriptorManager->GetResourceBindHeapInfo());
+		}
 	}
 
 	if (OptionalQueryPool)
diff --git a/Engine/Source/Runtime/VulkanRHI/Private/VulkanContext.cpp b/Engine/Source/Runtime/VulkanRHI/Private/VulkanContext.cpp
--- a/Engine/Source/Runtime/VulkanRHI/Private/VulkanContext.cpp
+++ b/Engine/Source/Runtime/VulkanRHI/Private/VulkanContext.cpp
@@ -719,8 +719,25 @@ void FVulkanCommandListContext::ApplyShaderBindingLayout(VkShaderStageFlags Shad
 			UBBindlessIndices[Index] = BindlessHandle.GetIndex();
 		}
 
-		VulkanRHI::vkCmdPushConstants( GetCommandBuffer().GetHandle(), Device.GetBindlessDescriptorManager()->GetPipelineLayout(),
-			ShaderStageFlags, 0 /*offset*/, sizeof(uint32) * ShaderBindingLayout->GetNumUniformBufferEntries(), UBBindlessIndices);
+		const uint32 DataSize = sizeof(uint32) * ShaderBindingLayout->GetNumUniformBufferEntries();
+		if (Device.GetBindlessDescriptorManager()->UseDescriptorHeaps())
+		{
+			const VkPushDataInfoEXT PushDataInfo = {
+				.sType = VK_STRUCTURE_TYPE_PUSH_DATA_INFO_EXT,
+				.pNext = nullptr,
+				.offset = 0,
+				.data = {
+					.address = UBBindlessIndices,
+					.size = DataSize
+				}
+			};
+			VulkanRHI::vkCmdPushDataEXT(GetCommandBuffer().GetHandle(), &PushDataInfo);
+		}
+		else
+		{
+			VulkanRHI::vkCmdPushConstants(GetCommandBuffer().GetHandle(), Device.GetBindlessDescriptorManager()->GetPipelineLayout(),
+				ShaderStageFlags, 0 /*offset*/, DataSize, UBBindlessIndices);
+		}
 	}
 }
 
diff --git a/Engine/Source/Runtime/VulkanRHI/Private/VulkanDescriptorSets.cpp b/Engine/Source/Runtime/VulkanRHI/Private/VulkanDescriptorSets.cpp
--- a/Engine/Source/Runtime/VulkanRHI/Private/VulkanDescriptorSets.cpp
+++ b/Engine/Source/Runtime/VulkanRHI/Private/VulkanDescriptorSets.cpp
@@ -1,6 +1,8 @@
 // Copyright Epic Games, Inc. All Rights Reserved.
 
 #include "VulkanDescriptorSets.h"
+#include "VulkanDevice.h"
+#include "VulkanBindlessDescriptorManager.h"
 
 TAutoConsoleVariable<int32> GDynamicGlobalUBs(
 	TEXT("r.Vulkan.DynamicGlobalUBs"),
@@ -537,3 +539,18 @@ void FVulkanDescriptorSetWriter::SetExcludedResources(const TArrayView<uint32> E
 	}
 #endif
 }
+
+bool FVulkanDescriptorSetWriter::WriteBindlessUniformBuffer(FVulkanDevice& Device, uint32 DescriptorIndex, const FVulkanUniformBuffer* UniformBuffer)
+{
+	if (Device.GetBindlessDescriptorManager()->UseDescriptorHeaps())
+	{
+		BindlessUBIndices[DescriptorIndex] = UniformBuffer->GetBindlessHandle().GetIndex();
+		return true;
+	}
+	else
+	{
+		const VulkanRHI::FVulkanAllocation& Allocation = UniformBuffer->GetAllocation();
+		const VkDeviceSize Range = UniformBuffer->IsUniformView() ? PLATFORM_MAX_UNIFORM_BUFFER_RANGE : UniformBuffer->GetSize();
+		return WriteUniformBuffer(DescriptorIndex, Allocation.GetBufferHandle(), Allocation.HandleId, UniformBuffer->GetOffset(), Range);
+	}
+}
diff --git a/Engine/Source/Runtime/VulkanRHI/Private/VulkanDescriptorSets.h b/Engine/Source/Runtime/VulkanRHI/Private/VulkanDescriptorSets.h
--- a/Engine/Source/Runtime/VulkanRHI/Private/VulkanDescriptorSets.h
+++ b/Engine/Source/Runtime/VulkanRHI/Private/VulkanDescriptorSets.h
@@ -529,6 +529,8 @@ public:
 		return NumWrites;
 	}
 
+	bool WriteBindlessUniformBuffer(FVulkanDevice& Device, uint32 DescriptorIndex, const FVulkanUniformBuffer* UniformBuffer);
+
 	bool WriteUniformBuffer(uint32 DescriptorIndex, VkBuffer BufferHandle, uint32 HandleId, VkDeviceSize Offset, VkDeviceSize Range)
 	{
 		return WriteBuffer<VK_DESCRIPTOR_TYPE_UNIFORM_BUFFER>(DescriptorIndex, BufferHandle, HandleId, Offset, Range);
@@ -597,7 +599,6 @@ public:
 		return WriteBufferView<VK_DESCRIPTOR_TYPE_UNIFORM_TEXEL_BUFFER>(DescriptorIndex, View);
 	}
 
-
 	bool WriteAccelerationStructure(uint32 DescriptorIndex, VkAccelerationStructureKHR InAccelerationStructure)
 	{
 		checkf(!UseVulkanDescriptorCache(), TEXT("Descriptor cache path for WriteAccelerationStructure() is not implemented"));
@@ -803,6 +804,9 @@ protected:
 
 	uint32 NumWrites;
 
+	// UB handles when running with bindless descriptor heaps
+	TArray<uint32> BindlessUBIndices;
+
 	FVulkanHashableDescriptorInfo* HashableDescriptorInfos;
 	mutable FVulkanDSetKey Key;
 	mutable bool bIsKeyDirty;
diff --git a/Engine/Source/Runtime/VulkanRHI/Private/VulkanPipeline.cpp b/Engine/Source/Runtime/VulkanRHI/Private/VulkanPipeline.cpp
--- a/Engine/Source/Runtime/VulkanRHI/Private/VulkanPipeline.cpp
+++ b/Engine/Source/Runtime/VulkanRHI/Private/VulkanPipeline.cpp
@@ -1504,24 +1504,48 @@ void FVulkanPipelineStateCacheManager::CreateEmptyFragmentStageLibrary()
 	VkPipelineShaderStageCreateInfo ShaderStageCreateInfo;
 	TArrayView<uint32> SpirvCode;
 	ANSICHAR EntryPoint[24];
+	VkShaderDescriptorSetAndBindingMappingInfoEXT BindingMappingInfo;
 	if (!FVulkanPlatform::SupportsNullPixelShader())
 	{
 		ZeroVulkanStruct(ModuleCreateInfo, VK_STRUCTURE_TYPE_SHADER_MODULE_CREATE_INFO);
 		ZeroVulkanStruct(ShaderStageCreateInfo, VK_STRUCTURE_TYPE_PIPELINE_SHADER_STAGE_CREATE_INFO);
 
 		FVulkanPixelShader* PixelShader = ResourceCast(TShaderMapRef<FNULLPS>(GetGlobalShaderMap(GMaxRHIFeatureLevel)).GetPixelShader());
-
 		SpirvCode = PixelShader->GetSpirvCode().GetCodeView();
-		ModuleCreateInfo.codeSize = SpirvCode.Num() * sizeof(uint32);
-		ModuleCreateInfo.pCode = SpirvCode.GetData();
+		PixelShader->GetEntryPoint(EntryPoint, 24);
 
-		ShaderStageCreateInfo.pNext = &ModuleCreateInfo;
-		ShaderStageCreateInfo.stage = VK_SHADER_STAGE_FRAGMENT_BIT;
+		ModuleCreateInfo = {
+			.sType = VK_STRUCTURE_TYPE_SHADER_MODULE_CREATE_INFO,
+			.pNext = nullptr,
+			.flags = 0,
+			.codeSize = SpirvCode.Num() * sizeof(uint32),
+			.pCode = SpirvCode.GetData()
+		};
 
-		PixelShader->GetEntryPoint(EntryPoint, 24);
-		ShaderStageCreateInfo.pName = EntryPoint;
+		ShaderStageCreateInfo = {
+			.sType = VK_STRUCTURE_TYPE_PIPELINE_SHADER_STAGE_CREATE_INFO,
+			.pNext = &ModuleCreateInfo,
+			.flags = 0,
+			.stage = VK_SHADER_STAGE_FRAGMENT_BIT,
+			.module = VK_NULL_HANDLE,
+			.pName = EntryPoint,
+			.pSpecializationInfo = nullptr // :todo:
+		};
 
 		GraphicsPipelineCreateInfo.pStages = &ShaderStageCreateInfo;
+
+		if (Device->GetBindlessDescriptorManager()->UseDescriptorHeaps())
+		{
+			TConstArrayView<VkDescriptorSetAndBindingMappingEXT> BindingMappings = Device->GetBindlessDescriptorManager()->GetBindingMappings(SF_Pixel);
+			BindingMappingInfo = {
+				.sType = VK_STRUCTURE_TYPE_SHADER_DESCRIPTOR_SET_AND_BINDING_MAPPING_INFO_EXT,
+				.pNext = nullptr,
+				.mappingCount = (uint32)BindingMappings.Num(),
+				.pMappings = BindingMappings.GetData()
+			};
+
+			AddToPNext(ShaderStageCreateInfo, BindingMappingInfo);
+		}
 	}
 
 	VERIFYVULKANRESULT(VulkanRHI::vkCreateGraphicsPipelines(Device->GetHandle(), VK_NULL_HANDLE, 1, &GraphicsPipelineCreateInfo, VULKAN_CPU_ALLOCATOR, &EmptyFragmentStageLibrary));
@@ -1612,11 +1636,11 @@ VkPipeline FVulkanPipelineStateCacheManager::FindOrCreateFragmentOutputStateLibr
 	DynamicStateCreateInfo.pDynamicStates = FragmentOutputStateDynamicStates;
 
 	const VkPipelineCreateFlags2CreateInfo Flags2CreateInfo = {
-	.sType = VK_STRUCTURE_TYPE_PIPELINE_CREATE_FLAGS_2_CREATE_INFO,
-	.pNext = &LibraryInfo,
-	.flags = VK_PIPELINE_CREATE_2_LIBRARY_BIT_KHR | 
-			 Device->GetBindlessDescriptorManager()->GetPipelineCreateFlag() |
-			 (GVulkanGraphicPipelineLibraryLinkingMode == 1 ? VK_PIPELINE_CREATE_2_RETAIN_LINK_TIME_OPTIMIZATION_INFO_BIT_EXT : 0ull)
+		.sType = VK_STRUCTURE_TYPE_PIPELINE_CREATE_FLAGS_2_CREATE_INFO,
+		.pNext = &LibraryInfo,
+		.flags = VK_PIPELINE_CREATE_2_LIBRARY_BIT_KHR | 
+				 Device->GetBindlessDescriptorManager()->GetPipelineCreateFlag() |
+				 (GVulkanGraphicPipelineLibraryLinkingMode == 1 ? VK_PIPELINE_CREATE_2_RETAIN_LINK_TIME_OPTIMIZATION_INFO_BIT_EXT : 0ull)
 	};
 
 	VkGraphicsPipelineCreateInfo GraphicsPipelineCreateInfo;
@@ -1765,6 +1789,7 @@ bool FVulkanPipelineStateCacheManager::CreateGfxPipelineFromEntry(FVulkanGraphic
 	ANSICHAR EntryPoints[ShaderStage::NumGraphicsStages][24];
 	VkPipelineShaderStageRequiredSubgroupSizeCreateInfo RequiredSubgroupSizeCreateInfo[ShaderStage::NumGraphicsStages];
 	VkSpecializationInfo SpecializationInfo[ShaderStage::NumGraphicsStages];
+	VkShaderDescriptorSetAndBindingMappingInfoEXT BindingMappingInfo[ShaderStage::NumGraphicsStages];
 	for (int32 ShaderStage = 0; ShaderStage < ShaderStage::NumGraphicsStages; ++ShaderStage)
 	{
 		if (!ShaderModules[ShaderStage].IsValid() || (Shaders[ShaderStage] == nullptr))
@@ -1796,10 +1821,10 @@ bool FVulkanPipelineStateCacheManager::CreateGfxPipelineFromEntry(FVulkanGraphic
 						.pNext = nullptr,
 						.requiredSubgroupSize = ShaderHeader.WaveSize
 					};
-				ShaderStages[PipelineInfo.stageCount].pNext = &RequiredSubgroupSizeCreateInfo[PipelineInfo.stageCount];
+					AddToPNext(ShaderStages[PipelineInfo.stageCount], RequiredSubgroupSizeCreateInfo[PipelineInfo.stageCount]);
+				}
 			}
 		}
-		}
 
 		const uint32 NumSpecConsts = PSO->SpecializationConstants[ShaderStage].Num();
 		checkf(NumSpecConsts < MaxNumSpecializationConstants, TEXT("Exceeded maximum number of specialization constants!  Raise FVulkanPipelineStateCacheManager::MaxNumSpecializationConstants."));
@@ -1815,6 +1840,19 @@ bool FVulkanPipelineStateCacheManager::CreateGfxPipelineFromEntry(FVulkanGraphic
 			ShaderStages[PipelineInfo.stageCount].pSpecializationInfo = &SpecializationInfo[ShaderStage];
 		}
 
+		if (Device->GetBindlessDescriptorManager()->UseDescriptorHeaps())
+		{
+			TConstArrayView<VkDescriptorSetAndBindingMappingEXT> BindingMappings = Device->GetBindlessDescriptorManager()->GetBindingMappings(ShaderStage::GetFrequencyForGfxStage(CurrStage));
+			BindingMappingInfo[PipelineInfo.stageCount] = {
+				.sType = VK_STRUCTURE_TYPE_SHADER_DESCRIPTOR_SET_AND_BINDING_MAPPING_INFO_EXT,
+				.pNext = nullptr,
+				.mappingCount = (uint32)BindingMappings.Num(),
+				.pMappings = BindingMappings.GetData()
+			};
+
+			AddToPNext(ShaderStages[PipelineInfo.stageCount], BindingMappingInfo[PipelineInfo.stageCount]);
+		}
+
 		PipelineInfo.stageCount++;
 	}
 
@@ -2991,6 +3029,21 @@ FVulkanComputePipeline* FVulkanPipelineStateCacheManager::CreateComputePipeline(
 		AddToPNext(PipelineInfo, Flags2CreateInfo);
 	}
 
+	VkShaderDescriptorSetAndBindingMappingInfoEXT BindingMappingInfo;
+	if (Device->GetBindlessDescriptorManager()->UseDescriptorHeaps())
+	{
+		TConstArrayView<VkDescriptorSetAndBindingMappingEXT> BindingMappings = Device->GetBindlessDescriptorManager()->GetBindingMappings(SF_Compute);
+		BindingMappingInfo = {
+			.sType = VK_STRUCTURE_TYPE_SHADER_DESCRIPTOR_SET_AND_BINDING_MAPPING_INFO_EXT,
+			.pNext = nullptr,
+			.mappingCount = (uint32)BindingMappings.Num(),
+			.pMappings = BindingMappings.GetData()
+		};
+
+		AddToPNext(PipelineInfo.stage, BindingMappingInfo);
+	}
+
+
 	VkResult Result;
 	{
 		QUICK_SCOPE_CYCLE_COUNTER(STAT_VulkanComputePSOCreate);
diff --git a/Engine/Source/Runtime/VulkanRHI/Private/VulkanPipelineState.cpp b/Engine/Source/Runtime/VulkanRHI/Private/VulkanPipelineState.cpp
--- a/Engine/Source/Runtime/VulkanRHI/Private/VulkanPipelineState.cpp
+++ b/Engine/Source/Runtime/VulkanRHI/Private/VulkanPipelineState.cpp
@@ -16,10 +16,6 @@
 #include "RHICoreShader.h"
 #include "GlobalRenderResources.h"  // For GBlackTexture
 
-enum
-{
-	NumAllocationsPerPool = 8,
-};
 
 #if UE_BUILD_DEBUG || UE_BUILD_DEVELOPMENT
 static TAutoConsoleVariable<int32> GAlwaysWriteDS(
@@ -132,7 +128,11 @@ void FVulkanCommonPipelineDescriptorState::CreateDescriptorWriteInfos()
 			CurrentAccelerationStructures,
 			DefaultSampler, DefaultImageView);
 
-		if (!bUseBindless)
+		if (bUseBindless)
+		{
+			DSWriter[Set].BindlessUBIndices.SetNum(StageInfo.Types.Num());
+		}
+		else
 		{
 			DSWriter[Set].SetExcludedResources(GetExcludedResources(Set));
 		}
@@ -323,42 +323,88 @@ void FVulkanComputePipelineDescriptorState::UpdateBindlessDescriptors(FVulkanCom
 		return;
 	}
 
-	// Create a local copy of the writes
-	TArray<VkWriteDescriptorSet, TInlineAllocator<VulkanBindless::MaxUniformBuffersPerStage>> UBWrites;
-	UBWrites.SetNumZeroed(Header.NumBoundUniformBuffers);
-	FMemory::Memcpy(UBWrites.GetData(), DSWriter[ShaderStage::EStage::Compute].WriteDescriptors, Header.NumBoundUniformBuffers * sizeof(VkWriteDescriptorSet));
-
-	VkDescriptorBufferInfo ConstantBufferInfo;
+	const VkDeviceSize UBOffsetAlignment = Device.GetLimits().minUniformBufferOffsetAlignment;
 
-	// UBs are currently set from a fresh batch of descriptors for every call, so ignore PackedUniformBuffersDirty
-	check(PackedUniformBufferState.Mask <= 1);
-	if (PackedUniformBufferState.Mask != 0)
+	auto GetGlobalUBAlloc = [&]()
 	{
-		const FPackedUniformBuffers::FPackedBuffer& StagedUniformBuffer = PackedUniformBufferState.UniformBuffers.GetBuffer();
-		const int32 UBSize = StagedUniformBuffer.Num();
-		const VkDeviceSize UBOffsetAlignment = Device.GetLimits().minUniformBufferOffsetAlignment;
-
 		VulkanRHI::FVulkanAllocation TempAllocation;
-		uint8* MappedPointer = Device.GetTempBlockAllocator().Alloc(UBSize, UBOffsetAlignment, Context, TempAllocation, nullptr);
-		FMemory::Memcpy(MappedPointer, StagedUniformBuffer.GetData(), UBSize);
-
-		VkWriteDescriptorSet& WriteDescriptorSet = UBWrites[0];
-		check(WriteDescriptorSet.dstBinding == 0);
-		check(WriteDescriptorSet.dstArrayElement == 0);
-		check(WriteDescriptorSet.descriptorCount == 1);
-		check(WriteDescriptorSet.descriptorType == VK_DESCRIPTOR_TYPE_UNIFORM_BUFFER);
-		
-		WriteDescriptorSet.pBufferInfo = &ConstantBufferInfo;
+		check(PackedUniformBufferState.Mask <= 1);
+		if (PackedUniformBufferState.Mask != 0)
+		{
+			const FPackedUniformBuffers::FPackedBuffer& StagedUniformBuffer = PackedUniformBufferState.UniformBuffers.GetBuffer();
+			const int32 UBSize = StagedUniformBuffer.Num();
 
-		ConstantBufferInfo.buffer = TempAllocation.GetBufferHandle();
-		ConstantBufferInfo.offset = TempAllocation.Offset;
-		ConstantBufferInfo.range = UBSize;
+			uint8* MappedPointer = Device.GetTempBlockAllocator().Alloc(UBSize, UBOffsetAlignment, Context, TempAllocation, nullptr);
+			FMemory::Memcpy(MappedPointer, StagedUniformBuffer.GetData(), UBSize);
+		}
+		return TempAllocation;
+	};
 
-		PackedUniformBufferState.Dirty = 0;
+	if (Device.GetBindlessDescriptorManager()->UseDescriptorHeaps())
+	{
+		TArray<uint32> PushData = DSWriter[ShaderStage::EStage::Compute].BindlessUBIndices;
+		check(PushData.Num() == Header.NumBoundUniformBuffers);
+
+		VulkanRHI::FVulkanAllocation TempAllocation = GetGlobalUBAlloc();
+		if (TempAllocation.IsValid())
+		{
+			// NOTE: Can be replaced with BufferAddr straight in push data when DescriptorBuffer is removes (instead of the temporary handle)
+			FVulkanBindlessDescriptorManager* Mgr = Device.GetBindlessDescriptorManager();
+			FRHIDescriptorHandle BindlessHandle = Mgr->AllocateDescriptor(ERHIDescriptorType::CBV);
+			Mgr->UpdateDescriptor({}, BindlessHandle, TempAllocation.GetBufferHandle(), TempAllocation.Offset, TempAllocation.Size);
+			PushData[0] = BindlessHandle.GetIndex();
+			Device.GetDeferredDeletionQueue().EnqueueBindlessHandle(BindlessHandle);
+		}
+
+		const VkPushDataInfoEXT PushDataInfo = {
+			.sType = VK_STRUCTURE_TYPE_PUSH_DATA_INFO_EXT,
+			.pNext = nullptr,
+			.offset = (FRHIShaderBindingLayout::MaxUniformBufferEntries * sizeof(uint32)),  // place them after the static UBs
+			.data = {
+				.address = PushData.GetData(),
+				.size = PushData.NumBytes()
+			}
+		};
+		VulkanRHI::vkCmdPushDataEXT(Context.GetCommandBuffer().GetHandle(), &PushDataInfo);
 	}
+	else
+	{
+		// Create a local copy of the writes
+		TArray<VkWriteDescriptorSet, TInlineAllocator<VulkanBindless::MaxUniformBuffersPerStage>> UBWrites;
+		UBWrites.SetNumZeroed(Header.NumBoundUniformBuffers);
+		FMemory::Memcpy(UBWrites.GetData(), DSWriter[ShaderStage::EStage::Compute].WriteDescriptors, Header.NumBoundUniformBuffers * sizeof(VkWriteDescriptorSet));
+
+		VkDescriptorBufferInfo ConstantBufferInfo;
+
+		// UBs are currently set from a fresh batch of descriptors for every call, so ignore PackedUniformBuffersDirty
+		check(PackedUniformBufferState.Mask <= 1);
+		if (PackedUniformBufferState.Mask != 0)
+		{
+			const FPackedUniformBuffers::FPackedBuffer& StagedUniformBuffer = PackedUniformBufferState.UniformBuffers.GetBuffer();
+			const int32 UBSize = StagedUniformBuffer.Num();
 
-	// Send to descriptor manager
-	Device.GetBindlessDescriptorManager()->RegisterUniformBuffers(Context, VK_SHADER_STAGE_COMPUTE_BIT, UBWrites);
+			VulkanRHI::FVulkanAllocation TempAllocation;
+			uint8* MappedPointer = Device.GetTempBlockAllocator().Alloc(UBSize, UBOffsetAlignment, Context, TempAllocation, nullptr);
+			FMemory::Memcpy(MappedPointer, StagedUniformBuffer.GetData(), UBSize);
+
+			VkWriteDescriptorSet& WriteDescriptorSet = UBWrites[0];
+			check(WriteDescriptorSet.dstBinding == 0);
+			check(WriteDescriptorSet.dstArrayElement == 0);
+			check(WriteDescriptorSet.descriptorCount == 1);
+			check(WriteDescriptorSet.descriptorType == VK_DESCRIPTOR_TYPE_UNIFORM_BUFFER);
+
+			WriteDescriptorSet.pBufferInfo = &ConstantBufferInfo;
+
+			ConstantBufferInfo.buffer = TempAllocation.GetBufferHandle();
+			ConstantBufferInfo.offset = TempAllocation.Offset;
+			ConstantBufferInfo.range = UBSize;
+
+			PackedUniformBufferState.Dirty = 0;
+		}
+
+		// Send to descriptor manager
+		Device.GetBindlessDescriptorManager()->RegisterUniformBuffers(Context, VK_SHADER_STAGE_COMPUTE_BIT, UBWrites);
+	}
 }
 
 TArray<uint32> FVulkanComputePipelineDescriptorState::GetExcludedResources(uint32 SetIndex) const
@@ -574,75 +620,157 @@ void FVulkanGraphicsPipelineDescriptorState::UpdateBindlessDescriptors(FVulkanCo
 
 	const VkDeviceSize UBOffsetAlignment = Device.GetLimits().minUniformBufferOffsetAlignment;
 
-	TArray<VkWriteDescriptorSet, TInlineAllocator<VulkanBindless::MaxUniformBuffersTotal>> UBCompactedWrites;
-	TStaticArray<VkDescriptorBufferInfo, ShaderStage::NumGraphicsStages> ConstantBufferInfos;
-
-	VkShaderStageFlags StageFlags = 0;
-
-#if VULKAN_ENABLE_AGGRESSIVE_STATS
-	SCOPE_CYCLE_COUNTER(STAT_VulkanApplyPackedUniformBuffers);
-#endif
-	for (int32 Stage = 0; Stage < ShaderStage::NumGraphicsStages; ++Stage)
+	auto GetGlobalUBAlloc = [&](int32 Stage)
 	{
-		const EShaderFrequency Frequency = GetFrequencyForGfxStage((ShaderStage::EStage)Stage);
-		const FVulkanShader* VulkanShader = GfxPipeline->GetVulkanShader(Frequency);
-		if (!VulkanShader)
+		VulkanRHI::FVulkanAllocation TempAllocation;
+
+		FVulkanPackedUniformBufferState& State = PackedUniformBufferStates[Stage];
+		check(State.Mask <= 1);
+		if (State.Mask != 0)
 		{
-			continue;
+			const FPackedUniformBuffers::FPackedBuffer& StagedUniformBuffer = State.UniformBuffers.GetBuffer();
+			const int32 UBSize = StagedUniformBuffer.Num();
+
+			uint8* MappedPointer = Device.GetTempBlockAllocator().Alloc(UBSize, UBOffsetAlignment, Context, TempAllocation, nullptr);
+			FMemory::Memcpy(MappedPointer, StagedUniformBuffer.GetData(), UBSize);
+
+			State.Dirty = 0;
 		}
 
-		StageFlags |= UEFrequencyToVKStageBit(Frequency);
+		return TempAllocation;
+	};
 
-		const FVulkanShaderHeader& Header = VulkanShader->GetCodeHeader();
+	if (Device.GetBindlessDescriptorManager()->UseDescriptorHeaps())
+	{
+		TStaticArray<uint32, VulkanBindless::MaxUniformBuffersTotal> PushData;
+		FMemory::Memzero(PushData);
 
-		// Move the writes to compacted binding indices
-		const int32 BindingOffset = VulkanBindless::GetOffsetForFrequency(Frequency);
-		for (int32 UBIndex = 0; UBIndex < (int32)Header.NumBoundUniformBuffers; ++UBIndex)
+		int32 PushDataCount = 0;
+		for (int32 Stage = 0; Stage < ShaderStage::NumGraphicsStages; ++Stage)
 		{
-			const VkWriteDescriptorSet& WriteDescriptorSet = DSWriter[Stage].WriteDescriptors[UBIndex];
-			check(WriteDescriptorSet.dstBinding == UBIndex);
-			check(WriteDescriptorSet.dstArrayElement == 0);
-			check(WriteDescriptorSet.descriptorCount == 1);
-			check(WriteDescriptorSet.descriptorType == VK_DESCRIPTOR_TYPE_UNIFORM_BUFFER);
+			const EShaderFrequency Frequency = GetFrequencyForGfxStage((ShaderStage::EStage)Stage);
+			const FVulkanShader* VulkanShader = GfxPipeline->GetVulkanShader(Frequency);
+			if (!VulkanShader)
+			{
+				continue;
+			}
+
+			const FVulkanShaderHeader& Header = VulkanShader->GetCodeHeader();
+			const TArray<uint32>& BindlessUBIndices = DSWriter[Stage].BindlessUBIndices;
+			check(BindlessUBIndices.Num() == Header.NumBoundUniformBuffers);
 
-			VkWriteDescriptorSet& CompactedWrite = UBCompactedWrites.Add_GetRef(WriteDescriptorSet);
+			if (Header.NumBoundUniformBuffers == 0)
+			{
+				continue;
+			}
+
+			// NOTE: All this shuffling around to fit everything in MaxUniformBuffersTotal won't be necessary when DescriptorBuffer support is removed
+			const int32 BindingOffset = VulkanBindless::GetOffsetForFrequency(Frequency);
 			if (BindingOffset < 0)
 			{
-				CompactedWrite.dstBinding = VulkanBindless::MaxUniformBuffersTotal - UBIndex - 1;
+				for (int32 UBIndex = 0; UBIndex < BindlessUBIndices.Num(); ++UBIndex)
+				{
+					const int32 PushDataIndex = VulkanBindless::MaxUniformBuffersTotal - UBIndex - 1;
+					PushData[PushDataIndex] = BindlessUBIndices[UBIndex];
+				}
+
+				PushDataCount = VulkanBindless::MaxUniformBuffersTotal;
 			}
 			else
 			{
-				CompactedWrite.dstBinding = BindingOffset + UBIndex;
+				FMemory::Memcpy(&PushData[BindingOffset], BindlessUBIndices.GetData(), BindlessUBIndices.NumBytes());
+				PushDataCount = FMath::Max(PushDataCount, BindingOffset + BindlessUBIndices.Num());
 			}
 
-			// UBs are currently set from a fresh batch of descriptors for every call, so ignore PackedUniformBuffersDirty
-			FVulkanPackedUniformBufferState& State = PackedUniformBufferStates[Stage];
+			VulkanRHI::FVulkanAllocation TempAllocation = GetGlobalUBAlloc(Stage);
+			if (TempAllocation.IsValid())
+			{
+				// NOTE: Can be replaced with BufferAddr straight in push data when DescriptorBuffer is removes
+				FVulkanBindlessDescriptorManager* Mgr = Device.GetBindlessDescriptorManager();
+				FRHIDescriptorHandle BindlessHandle = Mgr->AllocateDescriptor(ERHIDescriptorType::CBV);
+				Mgr->UpdateDescriptor({}, BindlessHandle, TempAllocation.GetBufferHandle(), TempAllocation.Offset, TempAllocation.Size);
+				PushData[(BindingOffset < 0) ? VulkanBindless::MaxUniformBuffersTotal - 1 : BindingOffset] = BindlessHandle.GetIndex();
+				Device.GetDeferredDeletionQueue().EnqueueBindlessHandle(BindlessHandle);
+			}
+		}
 
-			check(State.Mask <= 1);
-			if (UBIndex == 0 && State.Mask != 0)
+		if (PushDataCount)
+		{
+			const VkPushDataInfoEXT PushDataInfo = {
+				.sType = VK_STRUCTURE_TYPE_PUSH_DATA_INFO_EXT,
+				.pNext = nullptr,
+				.offset = (FRHIShaderBindingLayout::MaxUniformBufferEntries * sizeof(uint32)),  // place them after the static UBs
+				.data = {
+					.address = PushData.GetData(),
+					.size = PushDataCount * sizeof(uint32)
+				}
+			};
+			VulkanRHI::vkCmdPushDataEXT(Context.GetCommandBuffer().GetHandle(), &PushDataInfo);
+		}
+	}
+	else
+	{
+		TArray<VkWriteDescriptorSet, TInlineAllocator<VulkanBindless::MaxUniformBuffersTotal>> UBCompactedWrites;
+		TStaticArray<VkDescriptorBufferInfo, ShaderStage::NumGraphicsStages> ConstantBufferInfos;
+
+		VkShaderStageFlags StageFlags = 0;
+
+#if VULKAN_ENABLE_AGGRESSIVE_STATS
+		SCOPE_CYCLE_COUNTER(STAT_VulkanApplyPackedUniformBuffers);
+#endif
+		for (int32 Stage = 0; Stage < ShaderStage::NumGraphicsStages; ++Stage)
+		{
+			const EShaderFrequency Frequency = GetFrequencyForGfxStage((ShaderStage::EStage)Stage);
+			const FVulkanShader* VulkanShader = GfxPipeline->GetVulkanShader(Frequency);
+			if (!VulkanShader)
 			{
-				const FPackedUniformBuffers::FPackedBuffer& StagedUniformBuffer = State.UniformBuffers.GetBuffer();
-				const int32 UBSize = StagedUniformBuffer.Num();
+				continue;
+			}
 
-				VulkanRHI::FVulkanAllocation TempAllocation;
-				uint8* MappedPointer = Device.GetTempBlockAllocator().Alloc(UBSize, UBOffsetAlignment, Context, TempAllocation, nullptr);
-				FMemory::Memcpy(MappedPointer, StagedUniformBuffer.GetData(), UBSize);
+			StageFlags |= UEFrequencyToVKStageBit(Frequency);
 
-				VkDescriptorBufferInfo& ConstantBufferInfo = ConstantBufferInfos[Stage];
+			const FVulkanShaderHeader& Header = VulkanShader->GetCodeHeader();
 
-				CompactedWrite.pBufferInfo = &ConstantBufferInfo;
+			// Move the writes to compacted binding indices
+			const int32 BindingOffset = VulkanBindless::GetOffsetForFrequency(Frequency);
+			for (int32 UBIndex = 0; UBIndex < (int32)Header.NumBoundUniformBuffers; ++UBIndex)
+			{
+				const VkWriteDescriptorSet& WriteDescriptorSet = DSWriter[Stage].WriteDescriptors[UBIndex];
+				check(WriteDescriptorSet.dstBinding == UBIndex);
+				check(WriteDescriptorSet.dstArrayElement == 0);
+				check(WriteDescriptorSet.descriptorCount == 1);
+				check(WriteDescriptorSet.descriptorType == VK_DESCRIPTOR_TYPE_UNIFORM_BUFFER);
+
+				VkWriteDescriptorSet& CompactedWrite = UBCompactedWrites.Add_GetRef(WriteDescriptorSet);
+				if (BindingOffset < 0)
+				{
+					CompactedWrite.dstBinding = VulkanBindless::MaxUniformBuffersTotal - UBIndex - 1;
+				}
+				else
+				{
+					CompactedWrite.dstBinding = BindingOffset + UBIndex;
+				}
 
-				ConstantBufferInfo.buffer = TempAllocation.GetBufferHandle();
-				ConstantBufferInfo.offset = TempAllocation.Offset;
-				ConstantBufferInfo.range = UBSize;
+				if (UBIndex == 0)
+				{
+					VulkanRHI::FVulkanAllocation TempAllocation = GetGlobalUBAlloc(Stage);
+					if (TempAllocation.IsValid())
+					{
+						VkDescriptorBufferInfo& ConstantBufferInfo = ConstantBufferInfos[Stage];
 
-				State.Dirty = 0;
+						CompactedWrite.pBufferInfo = &ConstantBufferInfo;
+
+						ConstantBufferInfo.buffer = TempAllocation.GetBufferHandle();
+						ConstantBufferInfo.offset = TempAllocation.Offset;
+						ConstantBufferInfo.range = TempAllocation.Size;
+					}
+				}
 			}
 		}
-	}
 
-	// Send to descriptor manager
-	Device.GetBindlessDescriptorManager()->RegisterUniformBuffers(Context, StageFlags, UBCompactedWrites);
+		// Send to descriptor manager
+		Device.GetBindlessDescriptorManager()->RegisterUniformBuffers(Context, StageFlags, UBCompactedWrites);
+	}
 }
 
 TArray<uint32> FVulkanGraphicsPipelineDescriptorState::GetExcludedResources(uint32 SetIndex) const
diff --git a/Engine/Source/Runtime/VulkanRHI/Private/VulkanPipelineState.h b/Engine/Source/Runtime/VulkanRHI/Private/VulkanPipelineState.h
--- a/Engine/Source/Runtime/VulkanRHI/Private/VulkanPipelineState.h
+++ b/Engine/Source/Runtime/VulkanRHI/Private/VulkanPipelineState.h
@@ -97,9 +97,14 @@ public:
 	inline void SetUniformBuffer(uint8 DescriptorSet, uint32 BindingIndex, const FVulkanUniformBuffer* UniformBuffer, bool bDynamic)
 	{
 		const VulkanRHI::FVulkanAllocation& Allocation = UniformBuffer->GetAllocation();
-		VkDeviceSize Range = UniformBuffer->IsUniformView() ? PLATFORM_MAX_UNIFORM_BUFFER_RANGE : UniformBuffer->GetSize();
+		const VkDeviceSize Range = UniformBuffer->IsUniformView() ? PLATFORM_MAX_UNIFORM_BUFFER_RANGE : UniformBuffer->GetSize();
 
-		if (bDynamic)
+		if (bUseBindless)
+		{
+			check(!bDynamic);
+			MarkDirty(DSWriter[DescriptorSet].WriteBindlessUniformBuffer(Device, BindingIndex, UniformBuffer));
+		}
+		else if (bDynamic)
 		{
 			MarkDirty(DSWriter[DescriptorSet].WriteDynamicUniformBuffer(BindingIndex, Allocation.GetBufferHandle(), Allocation.HandleId, 0, Range, UniformBuffer->GetOffset()));
 		}
diff --git a/Engine/Source/Runtime/VulkanRHI/Private/VulkanRayTracing.cpp b/Engine/Source/Runtime/VulkanRHI/Private/VulkanRayTracing.cpp
--- a/Engine/Source/Runtime/VulkanRHI/Private/VulkanRayTracing.cpp
+++ b/Engine/Source/Runtime/VulkanRHI/Private/VulkanRayTracing.cpp
@@ -596,7 +596,10 @@ void FVulkanRayTracingGeometry::SetupHitGroupSystemParameters()
 	{
 		if (Buffer)
 		{
-			return BindlessDescriptorManager->AllocateDescriptor(ERHIDescriptorType::BufferUAV, Buffer, ExtraOffset);
+			FRHIDescriptorHandle BindlessHandle = BindlessDescriptorManager->AllocateDescriptor(ERHIDescriptorType::BufferUAV);
+			check(ExtraOffset < Buffer->GetCurrentSize());
+			BindlessDescriptorManager->UpdateDescriptor({}, BindlessHandle, Buffer->GetHandle(), Buffer->GetOffset() + ExtraOffset, Buffer->GetCurrentSize() - ExtraOffset, true);
+			return BindlessHandle;
 		}
 		return FRHIDescriptorHandle();
 	};
@@ -631,17 +634,15 @@ void FVulkanRayTracingGeometry::SetupHitGroupSystemParameters()
 
 void FVulkanRayTracingGeometry::ReleaseBindlessHandles()
 {
-	FVulkanBindlessDescriptorManager* BindlessDescriptorManager = Device->GetBindlessDescriptorManager();
-
 	for (FRHIDescriptorHandle BindlesHandle : HitGroupSystemVertexViews)
 	{
-		BindlessDescriptorManager->FreeDescriptor(BindlesHandle);
+		Device->GetDeferredDeletionQueue().EnqueueBindlessHandle(BindlesHandle);
 	}
 	HitGroupSystemVertexViews.Reset(Initializer.Segments.Num());
 
 	if (HitGroupSystemIndexView.IsValid())
 	{
-		BindlessDescriptorManager->FreeDescriptor(HitGroupSystemIndexView);
+		Device->GetDeferredDeletionQueue().EnqueueBindlessHandle(HitGroupSystemIndexView);
 		HitGroupSystemIndexView = FRHIDescriptorHandle();
 	}
 }
@@ -757,11 +758,12 @@ void FVulkanRayTracingScene::BindBuffer(FRHIBuffer* InBuffer, uint32 InBufferOff
 		);
 
 #if VULKAN_USE_DEBUG_NAMES
-		VkAccelerationStructureKHR NativeAccelerationStructureHandle = View->GetAccelerationStructureView().Handle;
-
-		FString DebugNameString = Initializer.DebugName.ToString();
-		DebugNameString = (DebugNameString.IsEmpty()) ? TEXT("TLAS") : DebugNameString;
-		VULKAN_SET_DEBUG_NAME(Device, VK_OBJECT_TYPE_ACCELERATION_STRUCTURE_KHR, NativeAccelerationStructureHandle, TEXT("%s"), *DebugNameString);
+		if (const VkAccelerationStructureKHR NativeAccelerationStructureHandle = View->GetAccelerationStructureView().Handle)
+		{
+			FString DebugNameString = Initializer.DebugName.ToString();
+			DebugNameString = (DebugNameString.IsEmpty()) ? TEXT("TLAS") : DebugNameString;
+			VULKAN_SET_DEBUG_NAME(Device, VK_OBJECT_TYPE_ACCELERATION_STRUCTURE_KHR, NativeAccelerationStructureHandle, TEXT("%s"), *DebugNameString);
+		}
 #endif
 	}
 }
@@ -1460,10 +1462,20 @@ public:
 					}
 				}
 
-				VkRayTracingPipelineInterfaceCreateInfoKHR RayTracingPipelineInterfaceCreateInfo;
-				ZeroVulkanStruct(RayTracingPipelineInterfaceCreateInfo, VK_STRUCTURE_TYPE_RAY_TRACING_PIPELINE_INTERFACE_CREATE_INFO_KHR);
-				RayTracingPipelineInterfaceCreateInfo.maxPipelineRayHitAttributeSize = Key.MaxAttributeSizeInBytes;
-				RayTracingPipelineInterfaceCreateInfo.maxPipelineRayPayloadSize = Key.MaxPayloadSizeInBytes;
+				const VkRayTracingPipelineInterfaceCreateInfoKHR RayTracingPipelineInterfaceCreateInfo = {
+					.sType = VK_STRUCTURE_TYPE_RAY_TRACING_PIPELINE_INTERFACE_CREATE_INFO_KHR,
+					.pNext = nullptr,
+					.maxPipelineRayPayloadSize = Key.MaxPayloadSizeInBytes,
+					.maxPipelineRayHitAttributeSize = Key.MaxAttributeSizeInBytes
+				};
+
+				const VkPipelineCreateFlags2CreateInfo Flags2CreateInfo = {
+					.sType = VK_STRUCTURE_TYPE_PIPELINE_CREATE_FLAGS_2_CREATE_INFO,
+					.pNext = nullptr,
+					.flags = VK_PIPELINE_CREATE_2_LIBRARY_BIT_KHR |
+							 Device.GetBindlessDescriptorManager()->GetPipelineCreateFlag() |
+							 (GRHIGlobals.RayTracing.SupportsLinearSweptSpheres ? VK_PIPELINE_CREATE_2_RAY_TRACING_ALLOW_SPHERES_AND_LINEAR_SWEPT_SPHERES_BIT_NV : 0ull)
+				};
 
 				VkRayTracingPipelineCreateInfoKHR RayTracingPipelineCreateInfo;
 				ZeroVulkanStruct(RayTracingPipelineCreateInfo, VK_STRUCTURE_TYPE_RAY_TRACING_PIPELINE_CREATE_INFO_KHR);
@@ -1474,16 +1486,8 @@ public:
 				RayTracingPipelineCreateInfo.pGroups = &ShaderGroup;
 				RayTracingPipelineCreateInfo.maxPipelineRayRecursionDepth = 1;
 				RayTracingPipelineCreateInfo.layout = Device.GetBindlessDescriptorManager()->GetPipelineLayout();
-				RayTracingPipelineCreateInfo.flags = VK_PIPELINE_CREATE_LIBRARY_BIT_KHR | VK_PIPELINE_CREATE_DESCRIPTOR_BUFFER_BIT_EXT;
-
-				VkPipelineCreateFlags2CreateInfoKHR PipelineFlags2;
-				ZeroVulkanStruct(PipelineFlags2, VK_STRUCTURE_TYPE_PIPELINE_CREATE_FLAGS_2_CREATE_INFO_KHR);
-				PipelineFlags2.flags = RayTracingPipelineCreateInfo.flags | VK_PIPELINE_CREATE_2_RAY_TRACING_ALLOW_SPHERES_AND_LINEAR_SWEPT_SPHERES_BIT_NV;
-
-				if (GRHIGlobals.RayTracing.SupportsLinearSweptSpheres)
-				{
-					RayTracingPipelineCreateInfo.pNext = &PipelineFlags2;
-				}
+				RayTracingPipelineCreateInfo.flags = 0; // using VkPipelineCreateFlags2CreateInfoKHR
+				RayTracingPipelineCreateInfo.pNext = &Flags2CreateInfo;
 
 				const VkResult Result = Device.GetPipelineStateCache()->CreateRayTracingPipeline(RayTracingPipelineCreateInfo, bPartial, PipelineLibrary->PipelineHandle);
 				if (!bPartial)
@@ -1634,15 +1638,26 @@ FVulkanRayTracingPipelineState::FVulkanRayTracingPipelineState(FVulkanDevice& In
 		Pipelines.Add(PipelineLibrary->PipelineHandle);
 	}
 
-	VkRayTracingPipelineInterfaceCreateInfoKHR RayTracingPipelineInterfaceCreateInfo;
-	ZeroVulkanStruct(RayTracingPipelineInterfaceCreateInfo, VK_STRUCTURE_TYPE_RAY_TRACING_PIPELINE_INTERFACE_CREATE_INFO_KHR);
-	RayTracingPipelineInterfaceCreateInfo.maxPipelineRayPayloadSize = Initializer.MaxPayloadSizeInBytes;
-	RayTracingPipelineInterfaceCreateInfo.maxPipelineRayHitAttributeSize = Initializer.MaxAttributeSizeInBytes;
+	const VkRayTracingPipelineInterfaceCreateInfoKHR RayTracingPipelineInterfaceCreateInfo = {
+		.sType = VK_STRUCTURE_TYPE_RAY_TRACING_PIPELINE_INTERFACE_CREATE_INFO_KHR,
+		.pNext = nullptr,
+		.maxPipelineRayPayloadSize = Initializer.MaxPayloadSizeInBytes,
+		.maxPipelineRayHitAttributeSize = Initializer.MaxAttributeSizeInBytes
+	};
+
+	const VkPipelineLibraryCreateInfoKHR PipelineLibraryCreateInfo = {
+		.sType = VK_STRUCTURE_TYPE_PIPELINE_LIBRARY_CREATE_INFO_KHR,
+		.pNext = nullptr,
+		.libraryCount = (uint32)Pipelines.Num(),
+		.pLibraries = Pipelines.GetData()
+	};
 
-	VkPipelineLibraryCreateInfoKHR PipelineLibraryCreateInfo;
-	ZeroVulkanStruct(PipelineLibraryCreateInfo, VK_STRUCTURE_TYPE_PIPELINE_LIBRARY_CREATE_INFO_KHR);
-	PipelineLibraryCreateInfo.libraryCount = Pipelines.Num();
-	PipelineLibraryCreateInfo.pLibraries = Pipelines.GetData();
+	const VkPipelineCreateFlags2CreateInfo Flags2CreateInfo = {
+		.sType = VK_STRUCTURE_TYPE_PIPELINE_CREATE_FLAGS_2_CREATE_INFO,
+		.pNext = nullptr,
+		.flags = Device.GetBindlessDescriptorManager()->GetPipelineCreateFlag() |
+				 (GRHIGlobals.RayTracing.SupportsLinearSweptSpheres ? VK_PIPELINE_CREATE_2_RAY_TRACING_ALLOW_SPHERES_AND_LINEAR_SWEPT_SPHERES_BIT_NV : 0ull)
+	};
 
 	VkRayTracingPipelineCreateInfoKHR RayTracingPipelineCreateInfo;
 	ZeroVulkanStruct(RayTracingPipelineCreateInfo, VK_STRUCTURE_TYPE_RAY_TRACING_PIPELINE_CREATE_INFO_KHR);
@@ -1650,17 +1665,9 @@ FVulkanRayTracingPipelineState::FVulkanRayTracingPipelineState(FVulkanDevice& In
 	RayTracingPipelineCreateInfo.pLibraryInterface = &RayTracingPipelineInterfaceCreateInfo;
 	RayTracingPipelineCreateInfo.maxPipelineRayRecursionDepth = 1;
 	RayTracingPipelineCreateInfo.layout = InDevice.GetBindlessDescriptorManager()->GetPipelineLayout();
-	RayTracingPipelineCreateInfo.flags = VK_PIPELINE_CREATE_DESCRIPTOR_BUFFER_BIT_EXT;
+	RayTracingPipelineCreateInfo.flags = 0; // using VkPipelineCreateFlags2CreateInfo
+	RayTracingPipelineCreateInfo.pNext = &Flags2CreateInfo;
 
-	VkPipelineCreateFlags2CreateInfoKHR PipelineFlags2;
-	ZeroVulkanStruct(PipelineFlags2, VK_STRUCTURE_TYPE_PIPELINE_CREATE_FLAGS_2_CREATE_INFO_KHR);
-	PipelineFlags2.flags = RayTracingPipelineCreateInfo.flags | VK_PIPELINE_CREATE_2_RAY_TRACING_ALLOW_SPHERES_AND_LINEAR_SWEPT_SPHERES_BIT_NV;
-
-	if (GRHIGlobals.RayTracing.SupportsLinearSweptSpheres)
-	{
-		RayTracingPipelineCreateInfo.pNext = &PipelineFlags2;
-	}
-	
 	VERIFYVULKANRESULT_EXPANDED(InDevice.GetPipelineStateCache()->CreateRayTracingPipeline(RayTracingPipelineCreateInfo, bIsPartialPipeline, Pipeline));
 
 	// Grab all shader handles for each stage
diff --git a/Engine/Source/Runtime/VulkanRHI/Private/VulkanState.cpp b/Engine/Source/Runtime/VulkanRHI/Private/VulkanState.cpp
--- a/Engine/Source/Runtime/VulkanRHI/Private/VulkanState.cpp
+++ b/Engine/Source/Runtime/VulkanRHI/Private/VulkanState.cpp
@@ -228,17 +228,28 @@ FVulkanSamplerState::FVulkanSamplerState(const VkSamplerCreateInfo& InInfo, FVul
 	, BindlessHandle()
 	, bIsImmutable(bInIsImmutable)
 {
-	VERIFYVULKANRESULT(VulkanRHI::vkCreateSampler(InDevice.GetHandle(), &InInfo, VULKAN_CPU_ALLOCATOR, &Sampler));
+	FVulkanBindlessDescriptorManager* BindlessDescriptorManager = InDevice.GetBindlessDescriptorManager();
+	checkf(BindlessDescriptorManager, TEXT("SamplerState is being created before creation of the vulkan bindless manager."));
 
-	if (UseVulkanDescriptorCache())
+	BindlessHandle = BindlessDescriptorManager->AllocateDescriptor(ERHIDescriptorType::Sampler);
+
+	if (BindlessDescriptorManager->UseDescriptorHeaps())
 	{
-		SamplerId = ++GVulkanSamplerHandleIdCounter;
+		BindlessDescriptorManager->UpdateDescriptor(BindlessHandle, InInfo);
 	}
+	else
+	{
+		VERIFYVULKANRESULT(VulkanRHI::vkCreateSampler(InDevice.GetHandle(), &InInfo, VULKAN_CPU_ALLOCATOR, &Sampler));
 
-	checkf(InDevice.GetBindlessDescriptorManager(), TEXT("SamplerState is being created before creation of the vulkan bindless manager."));
-	if (InDevice.GetBindlessDescriptorManager()->GetConfiguration() != ERHIBindlessConfiguration::Disabled)
+		if (BindlessDescriptorManager->UseDescriptorBuffers())
+		{
+			BindlessDescriptorManager->UpdateDescriptor(BindlessHandle, Sampler);
+		}
+	}
+
+	if (UseVulkanDescriptorCache())
 	{
-		BindlessHandle = InDevice.GetBindlessDescriptorManager()->AllocateDescriptor(Sampler);
+		SamplerId = ++GVulkanSamplerHandleIdCounter;
 	}
 }
 
diff --git a/Engine/Source/Runtime/VulkanRHI/Private/VulkanTexture.cpp b/Engine/Source/Runtime/VulkanRHI/Private/VulkanTexture.cpp
--- a/Engine/Source/Runtime/VulkanRHI/Private/VulkanTexture.cpp
+++ b/Engine/Source/Runtime/VulkanRHI/Private/VulkanTexture.cpp
@@ -2258,7 +2258,6 @@ void FVulkanTexture::InitViews(VkImageUsageFlags ViewUsageFlags, VkSamplerYcbcrC
 	if (Image != VK_NULL_HANDLE && GetViewType() != VK_IMAGE_VIEW_TYPE_MAX_ENUM)
 	{
 		const FVulkanTextureViewDesc ViewDesc = CreateTextureViewDesc(this, ViewUsageFlags, SamplerYcbcrConversion);
-		const ERHIDescriptorType DescriptorType = SupportsSampling() ? ERHIDescriptorType::TextureSRV : ERHIDescriptorType::TextureUAV;
 
 		if (DefaultView)
 		{
@@ -2267,6 +2266,10 @@ void FVulkanTexture::InitViews(VkImageUsageFlags ViewUsageFlags, VkSamplerYcbcrC
 		}
 		else
 		{
+			const ERHIDescriptorType DescriptorType =
+				(FullAspectMask != PartialAspectMask) ? ERHIDescriptorType::Invalid :
+				SupportsSampling() ? ERHIDescriptorType::TextureSRV : ERHIDescriptorType::TextureUAV;
+
 			DefaultView = new FVulkanView(*Device, DescriptorType);
 			DefaultView->CreateTextureView(Image, ViewDesc);
 		}
@@ -2290,6 +2293,9 @@ void FVulkanTexture::InitViews(VkImageUsageFlags ViewUsageFlags, VkSamplerYcbcrC
 			}
 			else
 			{
+				const ERHIDescriptorType DescriptorType =
+					SupportsSampling() ? ERHIDescriptorType::TextureSRV : ERHIDescriptorType::TextureUAV;
+
 				PartialView = new FVulkanView(*Device, DescriptorType);
 				PartialView->CreateTextureView(Image, PartialViewDesc);
 			}
diff --git a/Engine/Source/Runtime/VulkanRHI/Private/VulkanUniformBuffer.cpp b/Engine/Source/Runtime/VulkanRHI/Private/VulkanUniformBuffer.cpp
--- a/Engine/Source/Runtime/VulkanRHI/Private/VulkanUniformBuffer.cpp
+++ b/Engine/Source/Runtime/VulkanRHI/Private/VulkanUniformBuffer.cpp
@@ -114,7 +114,8 @@ FVulkanUniformBuffer::FVulkanUniformBuffer(FVulkanDevice& InDevice, const FRHIUn
 	if (InLayout->ConstantBufferSize > 0)
 	{
 		AllocateMemory(Contents);
-		BindlessHandle = Device.GetBindlessDescriptorManager()->AllocateDescriptor(ERHIDescriptorType::CBV, GetDeviceAddress(), GetSize());
+		BindlessHandle = Device.GetBindlessDescriptorManager()->AllocateDescriptor(ERHIDescriptorType::CBV);
+		Device.GetBindlessDescriptorManager()->UpdateDescriptor({}, BindlessHandle, GetDeviceAddress(), GetSize(), true);
 	}
 }
 
diff --git a/Engine/Source/Runtime/VulkanRHI/Private/VulkanView.cpp b/Engine/Source/Runtime/VulkanRHI/Private/VulkanView.cpp
--- a/Engine/Source/Runtime/VulkanRHI/Private/VulkanView.cpp
+++ b/Engine/Source/Runtime/VulkanRHI/Private/VulkanView.cpp
@@ -91,7 +91,7 @@ void FVulkanView::InitAsTypedBufferView(const FVulkanContextArray& Contexts, FVu
 	const uint32 TotalOffset = Buffer->GetOffset() + InOffset;
 
 	check(UEFormat != PF_Unknown);
-	VkFormat Format = GVulkanBufferFormat[UEFormat];
+	const VkFormat Format = GVulkanBufferFormat[UEFormat];
 	check(Format != VK_FORMAT_UNDEFINED);
 
 	VkBufferViewCreateInfo ViewInfo;
@@ -221,8 +221,23 @@ void FVulkanView::InitAsTextureView(const FVulkanContextArray& Contexts, VkImage
 
 	if (BindlessHandle.IsValid())
 	{
-		const bool bDepthOrStencilAspect = (ViewDesc.AspectFlags & (VK_IMAGE_ASPECT_DEPTH_BIT | VK_IMAGE_ASPECT_STENCIL_BIT)) != 0;
-		Device.GetBindlessDescriptorManager()->UpdateDescriptor(Contexts, BindlessHandle, TV.View, bDepthOrStencilAspect, bImmediateUpdate);
+		if (Device.GetBindlessDescriptorManager()->UseDescriptorHeaps())
+		{
+			const bool bIsDepthStencil = VKHasAnyFlags(ViewDesc.AspectFlags, (VK_IMAGE_ASPECT_DEPTH_BIT | VK_IMAGE_ASPECT_STENCIL_BIT));
+			const VkImageDescriptorInfoEXT ImageDescriptorInfo = {
+				.sType = VK_STRUCTURE_TYPE_IMAGE_DESCRIPTOR_INFO_EXT,
+				.pNext = nullptr,
+				.pView = &ViewInfo,
+				.layout = (DescriptorType == ERHIDescriptorType::TextureUAV) ? VK_IMAGE_LAYOUT_GENERAL : bIsDepthStencil ? VK_IMAGE_LAYOUT_READ_ONLY_OPTIMAL : VK_IMAGE_LAYOUT_SHADER_READ_ONLY_OPTIMAL
+			};
+
+			Device.GetBindlessDescriptorManager()->UpdateDescriptor(Contexts, BindlessHandle, ImageDescriptorInfo, bImmediateUpdate);
+		}
+		else
+		{
+			const bool bDepthOrStencilAspect = (ViewDesc.AspectFlags & (VK_IMAGE_ASPECT_DEPTH_BIT | VK_IMAGE_ASPECT_STENCIL_BIT)) != 0;
+			Device.GetBindlessDescriptorManager()->UpdateDescriptor(Contexts, BindlessHandle, TV.View, bDepthOrStencilAspect, bImmediateUpdate);
+		}
 	}
 }
 
@@ -268,16 +283,26 @@ void FVulkanView::InitAsAccelerationStructureView(const FVulkanContextArray& Con
 	Storage.Emplace<FAccelerationStructureView>();
 	FAccelerationStructureView& ASV = Storage.Get<FAccelerationStructureView>();
 
-	VkAccelerationStructureCreateInfoKHR CreateInfo;
-	ZeroVulkanStruct(CreateInfo, VK_STRUCTURE_TYPE_ACCELERATION_STRUCTURE_CREATE_INFO_KHR);
-	CreateInfo.buffer = Buffer->GetHandle();
-	CreateInfo.offset = Buffer->GetOffset() + Offset;
-	CreateInfo.size = Size;
-	CreateInfo.type = VK_ACCELERATION_STRUCTURE_TYPE_TOP_LEVEL_KHR;
-
+	const VkAccelerationStructureCreateInfoKHR CreateInfo = {
+		.sType = VK_STRUCTURE_TYPE_ACCELERATION_STRUCTURE_CREATE_INFO_KHR,
+		.pNext = nullptr,
+		.createFlags = 0,
+		.buffer = Buffer->GetHandle(),
+		.offset = Buffer->GetOffset() + Offset,
+		.size = Size,
+		.type = VK_ACCELERATION_STRUCTURE_TYPE_TOP_LEVEL_KHR,
+		.deviceAddress = 0
+	};
 	VERIFYVULKANRESULT(VulkanDynamicAPI::vkCreateAccelerationStructureKHR(Device.GetHandle(), &CreateInfo, VULKAN_CPU_ALLOCATOR, &ASV.Handle));
 
-	Device.GetBindlessDescriptorManager()->UpdateDescriptor(Contexts, BindlessHandle, ASV.Handle, true);
+	if (Device.GetBindlessDescriptorManager()->UseDescriptorHeaps())
+	{
+		Device.GetBindlessDescriptorManager()->UpdateDescriptor(Contexts, BindlessHandle, Buffer->GetHandle(), Buffer->GetOffset() + Offset, Size, true);
+	}
+	else
+	{
+		Device.GetBindlessDescriptorManager()->UpdateDescriptor(Contexts, BindlessHandle, ASV.Handle, true);
+	}
 }
 
 //
