VulkanRHI: Fix acceleration structure descriptors with descriptor heaps. (WIP)

#rnx

#Synced-CL 52305081

diff --git a/Engine/Source/Runtime/VulkanRHI/Private/VulkanView.cpp b/Engine/Source/Runtime/VulkanRHI/Private/VulkanView.cpp
--- a/Engine/Source/Runtime/VulkanRHI/Private/VulkanView.cpp
+++ b/Engine/Source/Runtime/VulkanRHI/Private/VulkanView.cpp
@@ -297,7 +297,18 @@ void FVulkanView::InitAsAccelerationStructureView(const FVulkanContextArray& Con
 
 	if (Device.GetBindlessDescriptorManager()->UseDescriptorHeaps())
 	{
-		Device.GetBindlessDescriptorManager()->UpdateDescriptor(Contexts, BindlessHandle, Buffer->GetHandle(), Buffer->GetOffset() + Offset, Size, true);
+		const VkAccelerationStructureDeviceAddressInfoKHR DeviceAddressInfo = {
+			.sType = VK_STRUCTURE_TYPE_ACCELERATION_STRUCTURE_DEVICE_ADDRESS_INFO_KHR,
+			.pNext = nullptr,
+			.accelerationStructure = ASV.Handle
+		};
+
+		const VkDeviceAddressRangeEXT DeviceAddressRange = {
+			.address = VulkanDynamicAPI::vkGetAccelerationStructureDeviceAddressKHR(Device.GetHandle(), &DeviceAddressInfo),
+			.size = 0
+		};
+
+		Device.GetBindlessDescriptorManager()->UpdateDescriptor(Contexts, BindlessHandle, DeviceAddressRange, true);
 	}
 	else
 	{
