VulkanRHI: Add descriptor heap binding mappings to ray tracing shaders. (WIP)

#rnx

#Synced-CL 52305081

diff --git a/Engine/Source/Runtime/VulkanRHI/Private/VulkanBindlessDescriptorManager.cpp b/Engine/Source/Runtime/VulkanRHI/Private/VulkanBindlessDescriptorManager.cpp
--- a/Engine/Source/Runtime/VulkanRHI/Private/VulkanBindlessDescriptorManager.cpp
+++ b/Engine/Source/Runtime/VulkanRHI/Private/VulkanBindlessDescriptorManager.cpp
@@ -452,8 +452,12 @@ void FVulkanBindlessDescriptorManager::InitDescriptorHeaps()
 
 	// Create the binding mappings
 	{
+		HeapOnlyBindingMappings.SetNum(2);
+		VkDescriptorSetAndBindingMappingEXT& SamplerHeapMapping = HeapOnlyBindingMappings[VulkanBindless::BindlessSamplerSet];
+		VkDescriptorSetAndBindingMappingEXT& ResourceHeapMapping = HeapOnlyBindingMappings[VulkanBindless::BindlessResourceSet];
+
 		// Setup a common mapping for our sampler and resources heaps (will be inserted in all arrays)
-		const VkDescriptorSetAndBindingMappingEXT SamplerHeapMapping = {
+		SamplerHeapMapping = {
 			.sType = VK_STRUCTURE_TYPE_DESCRIPTOR_SET_AND_BINDING_MAPPING_EXT,
 			.pNext = nullptr,
 			.descriptorSet = VulkanBindless::BindlessSamplerSet,
@@ -476,7 +480,7 @@ void FVulkanBindlessDescriptorManager::InitDescriptorHeaps()
 			VK_SPIRV_RESOURCE_TYPE_UNIFORM_BUFFER_BIT_EXT | VK_SPIRV_RESOURCE_TYPE_READ_ONLY_STORAGE_BUFFER_BIT_EXT | VK_SPIRV_RESOURCE_TYPE_READ_WRITE_STORAGE_BUFFER_BIT_EXT |
 			(bHasRaytracingExtensions ? VK_SPIRV_RESOURCE_TYPE_ACCELERATION_STRUCTURE_BIT_EXT : 0);
 
-		const VkDescriptorSetAndBindingMappingEXT ResourceHeapMapping = {
+		ResourceHeapMapping = {
 			.sType = VK_STRUCTURE_TYPE_DESCRIPTOR_SET_AND_BINDING_MAPPING_EXT,
 			.pNext = nullptr,
 			.descriptorSet = VulkanBindless::BindlessResourceSet,
@@ -1301,8 +1305,12 @@ TConstArrayView<VkDescriptorSetAndBindingMappingEXT> FVulkanBindlessDescriptorMa
 	{
 		return GraphicsBindingMappings[Frequency];
 	}
-	else
+	else if ((Frequency == SF_Compute) || (Frequency == SF_RayGen))
 	{
 		return ComputeBindingMappings;
 	}
+	else
+	{
+		return HeapOnlyBindingMappings;
+	}
 }
diff --git a/Engine/Source/Runtime/VulkanRHI/Private/VulkanBindlessDescriptorManager.h b/Engine/Source/Runtime/VulkanRHI/Private/VulkanBindlessDescriptorManager.h
--- a/Engine/Source/Runtime/VulkanRHI/Private/VulkanBindlessDescriptorManager.h
+++ b/Engine/Source/Runtime/VulkanRHI/Private/VulkanBindlessDescriptorManager.h
@@ -95,12 +95,12 @@ public:
 	// Returns descriptor set layouts for a single stage
 	TConstArrayView<VkDescriptorSetLayout> GetDescriptorSetLayouts() const
 	{
-		return BindlessDescriptorSetLayouts;
+		return UseDescriptorBuffers() ? BindlessDescriptorSetLayouts : TConstArrayView<VkDescriptorSetLayout>();
 	}
 
 	TConstArrayView<VkPushConstantRange> GetPushConstantRanges() const
 	{
-		return UseDescriptorHeaps() ? TConstArrayView<VkPushConstantRange>() : PushConstantRanges;
+		return UseDescriptorBuffers() ? PushConstantRanges : TConstArrayView<VkPushConstantRange>();
 	}
 
 	bool UseDescriptorHeaps() const
@@ -172,6 +172,7 @@ private:
 	VkBindHeapInfoEXT ResourceBindHeapInfo;
 	TArray<VkDescriptorSetAndBindingMappingEXT> GraphicsBindingMappings[SF_NumGraphicsFrequencies];
 	TArray<VkDescriptorSetAndBindingMappingEXT> ComputeBindingMappings;
+	TArray<VkDescriptorSetAndBindingMappingEXT> HeapOnlyBindingMappings;
 
 	// Descriptor Heap functions
 	void InitDescriptorHeaps();
diff --git a/Engine/Source/Runtime/VulkanRHI/Private/VulkanRayTracing.cpp b/Engine/Source/Runtime/VulkanRHI/Private/VulkanRayTracing.cpp
--- a/Engine/Source/Runtime/VulkanRHI/Private/VulkanRayTracing.cpp
+++ b/Engine/Source/Runtime/VulkanRHI/Private/VulkanRayTracing.cpp
@@ -1401,6 +1401,8 @@ public:
 
 		auto CreatePipelineLibraryTask = [this, PipelineLibrary, Shader, Key, bPartial=Initializer.bPartial]()
 			{
+				const EShaderFrequency Frequency = Shader->GetFrequency();
+
 				ANSICHAR EntryPoint[MaxEntryPointNameLength];
 				Shader->GetEntryPoint(EntryPoint, MaxEntryPointNameLength);
 
@@ -1408,7 +1410,7 @@ public:
 				TStaticArray<VkPipelineShaderStageCreateInfo, MaxHitGroupShaderCount> ShaderStages;
 				ZeroVulkanStruct(ShaderStages[ShaderCount], VK_STRUCTURE_TYPE_PIPELINE_SHADER_STAGE_CREATE_INFO);
 				ShaderStages[ShaderCount].module = Key.ShaderHandle;
-				ShaderStages[ShaderCount].stage = UEFrequencyToVKStageBit(Shader->GetFrequency());  // will default to closestHitShader for SF_RayHitGroup
+				ShaderStages[ShaderCount].stage = UEFrequencyToVKStageBit(Frequency);  // will default to closestHitShader for SF_RayHitGroup
 				ShaderStages[ShaderCount].pName = EntryPoint;
 
 				VkRayTracingShaderGroupCreateInfoKHR ShaderGroup;
@@ -1477,6 +1479,24 @@ public:
 							 (GRHIGlobals.RayTracing.SupportsLinearSweptSpheres ? VK_PIPELINE_CREATE_2_RAY_TRACING_ALLOW_SPHERES_AND_LINEAR_SWEPT_SPHERES_BIT_NV : 0ull)
 				};
 
+				TStaticArray<VkShaderDescriptorSetAndBindingMappingInfoEXT, MaxHitGroupShaderCount> BindingMappingInfos;
+				if (Device.GetBindlessDescriptorManager()->UseDescriptorHeaps())
+				{
+					TConstArrayView<VkDescriptorSetAndBindingMappingEXT> BindingMappings = Device.GetBindlessDescriptorManager()->GetBindingMappings(Frequency);
+					for (uint32 StageIndex=0; StageIndex < ShaderCount; ++StageIndex)
+					{
+						BindingMappingInfos[StageIndex] = {
+							.sType = VK_STRUCTURE_TYPE_SHADER_DESCRIPTOR_SET_AND_BINDING_MAPPING_INFO_EXT,
+							.pNext = nullptr,
+							.mappingCount = (uint32)BindingMappings.Num(),
+							.pMappings = BindingMappings.GetData()
+						};
+
+						VkPipelineShaderStageCreateInfo& StageCreateInfo = ShaderStages[StageIndex];
+						AddToPNext(StageCreateInfo, BindingMappingInfos[StageIndex]);
+					}
+				}
+
 				VkRayTracingPipelineCreateInfoKHR RayTracingPipelineCreateInfo;
 				ZeroVulkanStruct(RayTracingPipelineCreateInfo, VK_STRUCTURE_TYPE_RAY_TRACING_PIPELINE_CREATE_INFO_KHR);
 				RayTracingPipelineCreateInfo.pLibraryInterface = &RayTracingPipelineInterfaceCreateInfo;
@@ -1941,17 +1961,47 @@ void FVulkanRayTracingCompactionRequestHandler::Update(FVulkanCommandListContext
 }
 
 
-static void SetRayGenResources(FVulkanDevice& Device, FVulkanCommandListContext& Context, const FRayTracingShaderBindings& InGlobalResourceBindings, FVulkanRayTracingShaderTable* ShaderTable)
+static void SetRayGenResources(FVulkanDevice& Device, FVulkanCommandListContext& Context, const FRayTracingShaderBindings& InGlobalResourceBindings)
 {
-	TArray<const FVulkanUniformBuffer*> UniformBuffers;
-	UniformBuffers.Reserve(UE_ARRAY_COUNT(InGlobalResourceBindings.UniformBuffers));
+	// Static uniform buffers
+	Context.ApplyShaderBindingLayout(VK_SHADER_STAGE_RAYGEN_BIT_KHR);
+
+	const uint32 MaxUniformBuffers = UE_ARRAY_COUNT(InGlobalResourceBindings.UniformBuffers);
 
 	// Uniform buffers
+	if (Device.GetBindlessDescriptorManager()->UseDescriptorHeaps())
+	{
+		TArray<uint32> PushData;
+		PushData.SetNumZeroed(MaxUniformBuffers);
+
+		uint32 MaxIndex = 0;
+		for (uint32 UBIndex = 0; UBIndex < MaxUniformBuffers; ++UBIndex)
+		{
+			const FVulkanUniformBuffer* UniformBuffer = ResourceCast(InGlobalResourceBindings.UniformBuffers[UBIndex]);
+			if (UniformBuffer)
+			{
+				check(UniformBuffer->GetBindlessHandle().IsValid());
+				PushData[UBIndex] = UniformBuffer->GetBindlessHandle().GetIndex();
+				MaxIndex = UBIndex;
+			}
+		}
+
+		const VkPushDataInfoEXT PushDataInfo = {
+			.sType = VK_STRUCTURE_TYPE_PUSH_DATA_INFO_EXT,
+			.pNext = nullptr,
+			.offset = (FRHIShaderBindingLayout::MaxUniformBufferEntries * sizeof(uint32)),  // place them after the static UBs
+			.data = {
+				.address = PushData.GetData(),
+				.size = (MaxIndex + 1) * sizeof(uint32)
+			}
+		};
+		VulkanRHI::vkCmdPushDataEXT(Context.GetCommandBuffer().GetHandle(), &PushDataInfo);
+	}
+	else 
 	{
 		TArray<VkWriteDescriptorSet, TInlineAllocator<VulkanBindless::MaxUniformBuffersPerStage>> UBWrites;
 		TArray<VkDescriptorBufferInfo, TInlineAllocator<VulkanBindless::MaxUniformBuffersPerStage>> UBBufferInfos;
 
-		const uint32 MaxUniformBuffers = UE_ARRAY_COUNT(InGlobalResourceBindings.UniformBuffers);
 		for (uint32 UBIndex = 0; UBIndex < MaxUniformBuffers; ++UBIndex)
 		{
 			const FVulkanUniformBuffer* UniformBuffer = ResourceCast(InGlobalResourceBindings.UniformBuffers[UBIndex]);
@@ -1968,24 +2018,11 @@ static void SetRayGenResources(FVulkanDevice& Device, FVulkanCommandListContext&
 				Write.descriptorType = VK_DESCRIPTOR_TYPE_UNIFORM_BUFFER;
 				Write.descriptorCount = 1;
 				Write.pBufferInfo = &BufferInfo;
-
-				UniformBuffers.AddUnique(UniformBuffer);
 			}
 		}
 
 		Device.GetBindlessDescriptorManager()->RegisterUniformBuffers(Context, VK_SHADER_STAGE_RAYGEN_BIT_KHR, UBWrites);
 	}
-
-	// Static uniform buffers
-	Context.ApplyShaderBindingLayout(VK_SHADER_STAGE_RAYGEN_BIT_KHR);
-
-	// Add all the UBs references by the shader table
-	TArrayView<TRefCountPtr<FRHIUniformBuffer>> ShaderTableUBs = ShaderTable->GetUBRefs();
-	for (TRefCountPtr<FRHIUniformBuffer>& UniformBuffer : ShaderTableUBs)
-	{
-		const FVulkanUniformBuffer* VulkanUniformBuffer = ResourceCast(UniformBuffer.GetReference());
-		UniformBuffers.AddUnique(VulkanUniformBuffer);
-	}
 }
 
 void FVulkanCommandListContext::RHIRayTraceDispatch(
@@ -2008,7 +2045,7 @@ void FVulkanCommandListContext::RHIRayTraceDispatch(
 	checkSlow(InRayGenShader->GetFrequency() == SF_RayGen);
 	const VkStridedDeviceAddressRegionKHR RayGenRegion = ShaderTable->CommitRayGenShader(*this, Pipeline->GetShaderIndex(RayGenShader), Pipeline->GetShaderHandles(SF_RayGen));
 
-	SetRayGenResources(Device, *this, InGlobalResourceBindings, ShaderTable);
+	SetRayGenResources(Device, *this, InGlobalResourceBindings);
 
 	VulkanRHI::vkCmdTraceRaysKHR(
 		CommandBuffer.GetHandle(),
@@ -2039,7 +2076,7 @@ void FVulkanCommandListContext::RHIRayTraceDispatchIndirect(
 	checkSlow(InRayGenShader->GetFrequency() == SF_RayGen);
 	const VkStridedDeviceAddressRegionKHR RayGenRegion = ShaderTable->CommitRayGenShader(*this, Pipeline->GetShaderIndex(RayGenShader), Pipeline->GetShaderHandles(SF_RayGen));
 
-	SetRayGenResources(Device, *this, InGlobalResourceBindings, ShaderTable);
+	SetRayGenResources(Device, *this, InGlobalResourceBindings);
 
 	FVulkanBuffer* ArgumentBuffer = ResourceCast(InArgumentBuffer);
 	const VkDeviceAddress IndirectDeviceAddress = ArgumentBuffer->GetDeviceAddress() + InArgumentOffset;
