VulkanRHI: Add missing flush in readback (memzero on host before the copy on GPU)

#rnx
#Synced-CL 52312524

diff --git a/Engine/Source/Runtime/VulkanRHI/Private/VulkanTexture.cpp b/Engine/Source/Runtime/VulkanRHI/Private/VulkanTexture.cpp
--- a/Engine/Source/Runtime/VulkanRHI/Private/VulkanTexture.cpp
+++ b/Engine/Source/Runtime/VulkanRHI/Private/VulkanTexture.cpp
@@ -1837,6 +1837,7 @@ FVulkanTexture::FVulkanTexture(FVulkanDevice& InDevice, const FRHITextureCreateD
 
 		void* Memory = Allocation.GetMappedPointer(Device);
 		FMemory::Memzero(Memory, Size);
+		Allocation.FlushMappedMemory(Device);
 
 		ImageOwnerType = EImageOwnerType::None;
 		ViewFormat = StorageFormat = UEToVkTextureFormat(InCreateDesc.Format, false);
