[VulkanRHI] Remove over-assertive deviceLUIDValid check from RHIGetGraphicsAdapterLUID

#rnx
#Synced-CL 52165013

diff --git a/Engine/Source/Runtime/VulkanRHI/Private/VulkanRHI.cpp b/Engine/Source/Runtime/VulkanRHI/Private/VulkanRHI.cpp
--- a/Engine/Source/Runtime/VulkanRHI/Private/VulkanRHI.cpp
+++ b/Engine/Source/Runtime/VulkanRHI/Private/VulkanRHI.cpp
@@ -1476,8 +1476,10 @@ uint64 FVulkanDynamicRHI::RHIGetGraphicsAdapterLUID(VkPhysicalDevice InPhysicalD
 	GpuProps2.pNext = &GpuIdProps;
 
 	VulkanRHI::vkGetPhysicalDeviceProperties2(InPhysicalDevice, &GpuProps2);
-	check(GpuIdProps.deviceLUIDValid);
-	AdapterLUID = reinterpret_cast<const uint64&>(GpuIdProps.deviceLUID);
+	if (GpuIdProps.deviceLUIDValid)
+	{
+		AdapterLUID = reinterpret_cast<const uint64&>(GpuIdProps.deviceLUID);
+	}
 #endif
 	return AdapterLUID;
 }
