VulkanRHI: use push address for globals (change binding0 to 8 byte value).  Track push constant state in the command buffer to reduce overall push data size.

#rnx

#Synced-CL 52393244

diff --git a/Engine/Source/Runtime/VulkanRHI/Private/VulkanBindlessDescriptorManager.cpp b/Engine/Source/Runtime/VulkanRHI/Private/VulkanBindlessDescriptorManager.cpp
--- a/Engine/Source/Runtime/VulkanRHI/Private/VulkanBindlessDescriptorManager.cpp
+++ b/Engine/Source/Runtime/VulkanRHI/Private/VulkanBindlessDescriptorManager.cpp
@@ -504,7 +504,11 @@ void FVulkanBindlessDescriptorManager::InitDescriptorHeaps()
 
 		for (int32 StageIndex = 0; StageIndex < SF_NumGraphicsFrequencies; ++StageIndex)
 		{
-			const int32 BindingOffset = VulkanBindless::GetOffsetForFrequency((EShaderFrequency)StageIndex);
+			const EShaderFrequency Frequency = (EShaderFrequency)StageIndex;
+			check(IsValidGraphicsFrequency(Frequency));
+
+			// Work around the limits imposed by shared support with descriptor_buffer and push_descriptors
+			const int32 PushDescriptorOffset = VulkanBindless::GetOffsetForFrequency(Frequency);
 			TArray<VkDescriptorSetAndBindingMappingEXT>& StageBindingMappings = GraphicsBindingMappings[StageIndex];
 
 			// Until our DXC supports SamplerHeapEXT / ResourceHeapEXT builtins, add bindings for sampler and resource heaps
@@ -514,20 +518,20 @@ void FVulkanBindlessDescriptorManager::InitDescriptorHeaps()
 
 			for (uint32 UBIndex = 0; UBIndex < (uint32)VulkanBindless::MaxUniformBuffersPerStage; ++UBIndex)
 			{
-				const uint32 BindingIndex = (BindingOffset < 0) ? (VulkanBindless::MaxUniformBuffersTotal - UBIndex - 1) : (BindingOffset + UBIndex);
+				const uint32 ShaderBindingIndex = (PushDescriptorOffset < 0) ? (VulkanBindless::MaxUniformBuffersTotal - UBIndex - 1) : (PushDescriptorOffset + UBIndex);
 
 				StageBindingMappings[UBIndex] = {
 					.sType = VK_STRUCTURE_TYPE_DESCRIPTOR_SET_AND_BINDING_MAPPING_EXT,
 					.pNext = nullptr,
 					.descriptorSet = VulkanBindless::BindlessSingleUseUniformBufferSet,
-					.firstBinding = BindingIndex,
+					.firstBinding = ShaderBindingIndex,
 					.bindingCount = 1,
 					.resourceMask = VK_SPIRV_RESOURCE_TYPE_UNIFORM_BUFFER_BIT_EXT,
 					.source = VK_DESCRIPTOR_MAPPING_SOURCE_HEAP_WITH_PUSH_INDEX_EXT,
 					.sourceData = {
 						.pushIndex = {
 							.heapOffset = 0,
-							.pushOffset = PushConstantOffset + (BindingIndex * 4u),  // as per spec, 4 byte offsets
+							.pushOffset = GetBindingPushDataOffset(Frequency, UBIndex),
 							.heapIndexStride = ResourceState.DescriptorSize,
 							.heapArrayStride = ResourceState.DescriptorSize,
 							.pEmbeddedSampler = nullptr,
@@ -540,6 +544,11 @@ void FVulkanBindlessDescriptorManager::InitDescriptorHeaps()
 					}
 				};
 			}
+
+			// Create another version of the mappings with a push address in slot 0
+			GraphicsBindingMappingsWithPushAddr[StageIndex] = GraphicsBindingMappings[StageIndex];
+			GraphicsBindingMappingsWithPushAddr[StageIndex][0].source = VK_DESCRIPTOR_MAPPING_SOURCE_PUSH_ADDRESS_EXT;
+			GraphicsBindingMappingsWithPushAddr[StageIndex][0].sourceData.pushAddressOffset = GetBindingPushDataOffset(Frequency, 0);
 		}
 
 		// Until our DXC supports SamplerHeapEXT / ResourceHeapEXT builtins, add bindings for sampler and resource heaps
@@ -559,7 +568,7 @@ void FVulkanBindlessDescriptorManager::InitDescriptorHeaps()
 				.sourceData = {
 					.pushIndex = {
 						.heapOffset = 0,
-						.pushOffset = PushConstantOffset + (UBIndex * 4u),  // as per spec, 4 byte increments
+						.pushOffset = GetBindingPushDataOffset(SF_Compute, UBIndex),
 						.heapIndexStride = ResourceState.DescriptorSize,
 						.heapArrayStride = ResourceState.DescriptorSize,
 						.pEmbeddedSampler = nullptr,
@@ -572,7 +581,49 @@ void FVulkanBindlessDescriptorManager::InitDescriptorHeaps()
 				}
 			};
 		}
+
+		// Create another version of the mappings with a push address in slot 0
+		ComputeBindingMappingsWithPushAddr = ComputeBindingMappings;
+		ComputeBindingMappingsWithPushAddr[0].source = VK_DESCRIPTOR_MAPPING_SOURCE_PUSH_ADDRESS_EXT;
+		ComputeBindingMappingsWithPushAddr[0].sourceData.pushAddressOffset = GetBindingPushDataOffset(SF_Compute, 0);
+	}
+}
+
+uint32 FVulkanBindlessDescriptorManager::GetBindingPushDataOffset(EShaderFrequency Frequency, uint32 BindingIndex)
+{
+	// Use by descriptor heaps (not bound by the limits of push descriptors)
+
+	// Bindings are placed in push data after the static uniform buffers
+	constexpr uint32 PushConstantOffset = (FRHIShaderBindingLayout::MaxUniformBufferEntries * sizeof(uint32));
+
+	// A full stage has MaxUniformBuffersPerStage indices.  
+	// The first one is 64bit in case an address is used directly (in the case of temporary buffers for globals)
+	// The other bindings all use 32bit indices into the resource heap
+	constexpr uint32 FullStageSize = sizeof(uint64) + ((VulkanBindless::MaxUniformBuffersPerStage - 1) * sizeof(uint32));
+	const uint32 BindingOffset = (BindingIndex == 0) ? 0u : 4u + (BindingIndex * sizeof(uint32));
+
+	switch (Frequency)
+	{
+	case SF_Vertex:			return PushConstantOffset + BindingOffset;
+	case SF_Mesh:			return PushConstantOffset + BindingOffset;
+
+	case SF_Pixel:			return PushConstantOffset + FullStageSize + BindingOffset;
+
+	case SF_Geometry:		return PushConstantOffset + (2 * FullStageSize) + BindingOffset;
+	case SF_Amplification:	return PushConstantOffset + (2 * FullStageSize) + BindingOffset;
+
+	case SF_Compute:		return PushConstantOffset + BindingOffset;
+
+	case SF_RayGen:
+	case SF_RayMiss:
+	case SF_RayHitGroup:
+	case SF_RayCallable:	return PushConstantOffset + BindingOffset;
+	default:
+		checkf(false, TEXT("Invalid shader frequency %d"), (int32)Frequency);
+		break;
 	}
+
+	return 0;
 }
 
 void FVulkanBindlessDescriptorManager::InitDescriptorBuffers()
@@ -1294,7 +1345,7 @@ void FVulkanBindlessDescriptorManager::UpdateDescriptorInternal(const FVulkanCon
 	}
 }
 
-TConstArrayView<VkDescriptorSetAndBindingMappingEXT> FVulkanBindlessDescriptorManager::GetBindingMappings(EShaderFrequency Frequency) const
+TConstArrayView<VkDescriptorSetAndBindingMappingEXT> FVulkanBindlessDescriptorManager::GetBindingMappings(EShaderFrequency Frequency, bool bHasGlobals) const
 {
 	if (UseDescriptorBuffers())
 	{
@@ -1303,11 +1354,11 @@ TConstArrayView<VkDescriptorSetAndBindingMappingEXT> FVulkanBindlessDescriptorMa
 
 	if (IsValidGraphicsFrequency(Frequency))
 	{
-		return GraphicsBindingMappings[Frequency];
+		return bHasGlobals ? GraphicsBindingMappingsWithPushAddr[Frequency] : GraphicsBindingMappings[Frequency];
 	}
 	else if ((Frequency == SF_Compute) || (Frequency == SF_RayGen))
 	{
-		return ComputeBindingMappings;
+		return bHasGlobals ? ComputeBindingMappingsWithPushAddr : ComputeBindingMappings;
 	}
 	else
 	{
diff --git a/Engine/Source/Runtime/VulkanRHI/Private/VulkanBindlessDescriptorManager.h b/Engine/Source/Runtime/VulkanRHI/Private/VulkanBindlessDescriptorManager.h
--- a/Engine/Source/Runtime/VulkanRHI/Private/VulkanBindlessDescriptorManager.h
+++ b/Engine/Source/Runtime/VulkanRHI/Private/VulkanBindlessDescriptorManager.h
@@ -30,6 +30,9 @@ public:
 		Count
 	};
 
+	static constexpr uint32 RequiredPushDataSize = 256u;
+	static uint32 GetBindingPushDataOffset(EShaderFrequency Frequency, uint32 BindingIndex);
+
 	FVulkanBindlessDescriptorManager(FVulkanDevice& InDevice);
 	~FVulkanBindlessDescriptorManager();
 
@@ -69,7 +72,7 @@ public:
 	void UpdateDescriptor(const FVulkanContextArray& Contexts, FRHIDescriptorHandle DescriptorHandle, VkAccelerationStructureKHR AccelerationStructure, bool bImmediateUpdate = true);
 	void RegisterUniformBuffers(FVulkanCommandListContext& Context, VkShaderStageFlags StageFlags, TConstArrayView<VkWriteDescriptorSet> UBDescriptors);
 
-	TConstArrayView<VkDescriptorSetAndBindingMappingEXT> GetBindingMappings(EShaderFrequency Frequency) const;
+	TConstArrayView<VkDescriptorSetAndBindingMappingEXT> GetBindingMappings(EShaderFrequency Frequency, bool bHasGlobals) const;
 
 	ERHIBindlessConfiguration GetConfiguration() const
 	{
@@ -170,9 +173,16 @@ private:
 	// Descriptor Heap specific values
 	VkBindHeapInfoEXT SamplerBindHeapInfo;
 	VkBindHeapInfoEXT ResourceBindHeapInfo;
+
+	// Only provides binding mappings for the Sampler and Resource Heaps
+	// NOTE: Won't be necessary when our DXC supports SamplerHeapEXT / ResourceHeapEXT builtins
+	TArray<VkDescriptorSetAndBindingMappingEXT> HeapOnlyBindingMappings;
+	// Binding mappins where only bindless handles are used (no temporary allocation for globals)
 	TArray<VkDescriptorSetAndBindingMappingEXT> GraphicsBindingMappings[SF_NumGraphicsFrequencies];
 	TArray<VkDescriptorSetAndBindingMappingEXT> ComputeBindingMappings;
-	TArray<VkDescriptorSetAndBindingMappingEXT> HeapOnlyBindingMappings;
+	// Bindings with a push address in Binding0 for the temporary allocation used for globals
+	TArray<VkDescriptorSetAndBindingMappingEXT> GraphicsBindingMappingsWithPushAddr[SF_NumGraphicsFrequencies];
+	TArray<VkDescriptorSetAndBindingMappingEXT> ComputeBindingMappingsWithPushAddr;
 
 	// Descriptor Heap functions
 	void InitDescriptorHeaps();
diff --git a/Engine/Source/Runtime/VulkanRHI/Private/VulkanCommandBuffer.cpp b/Engine/Source/Runtime/VulkanRHI/Private/VulkanCommandBuffer.cpp
--- a/Engine/Source/Runtime/VulkanRHI/Private/VulkanCommandBuffer.cpp
+++ b/Engine/Source/Runtime/VulkanRHI/Private/VulkanCommandBuffer.cpp
@@ -39,8 +39,6 @@ void FVulkanCommandBuffer::AllocMemory()
 {
 	// Assumes we are inside a lock for the pool
 	check(State == EState::NotAllocated);
-	CurrentViewports.Empty();
-	CurrentScissors.Empty();
 
 	VkCommandBufferAllocateInfo CreateCmdBufInfo;
 	ZeroVulkanStruct(CreateCmdBufInfo, VK_STRUCTURE_TYPE_COMMAND_BUFFER_ALLOCATE_INFO);
@@ -57,9 +55,7 @@ void FVulkanCommandBuffer::AllocMemory()
 	bHasStencilRef = 0;
 	State = EState::ReadyForBegin;
 
-	CurrentPipelineHandle = VK_NULL_HANDLE;
-	CurrentComputeShaderObject = VK_NULL_HANDLE;
-	FMemory::Memzero(CurrentGraphicsShaderObjects);
+	ResetCurrentState();
 
 	INC_DWORD_STAT(STAT_VulkanNumCmdBuffers);
 
@@ -333,6 +329,21 @@ void FVulkanCommandBuffer::SetUnusedGraphicsDynamicStates()
 	}
 }
 
+void FVulkanCommandBuffer::ResetCurrentState()
+{
+	CurrentViewports.Empty();
+	CurrentScissors.Empty();
+	CurrentStencilRef = 0;
+
+	CurrentPipelineHandle = VK_NULL_HANDLE;
+	CurrentComputeShaderObject = VK_NULL_HANDLE;
+	FMemory::Memzero(CurrentGraphicsShaderObjects);
+
+	CurrentPushConstants.SetNumUninitialized(FVulkanBindlessDescriptorManager::RequiredPushDataSize / sizeof(uint32));
+	FMemory::Memset(CurrentPushConstants.GetData(), 0xFF, CurrentPushConstants.NumBytes());
+}
+
+
 void FVulkanCommandBuffer::Reset()
 {
 	// Reset the secondary command buffers we executed from this one
@@ -352,13 +363,7 @@ void FVulkanCommandBuffer::Reset()
 		bHasScissor = false;
 		bHasStencilRef = false;
 
-		CurrentViewports.Empty();
-		CurrentScissors.Empty();
-		CurrentStencilRef = 0;
-
-		CurrentPipelineHandle = VK_NULL_HANDLE;
-		CurrentComputeShaderObject = VK_NULL_HANDLE;
-		FMemory::Memzero(CurrentGraphicsShaderObjects);
+		ResetCurrentState();
 
 		for (VkEvent BarrierEvent : EndedBarrierEvents)
 		{
@@ -472,6 +477,106 @@ void FVulkanCommandBuffer::EndSplitBarrier(VkEvent BarrierEvent, const VkDepende
 	EndedBarrierEvents.Add(BarrierEvent);
 }
 
+void FVulkanCommandBuffer::SetStaticUniformBufferPushConstants(VkShaderStageFlags ShaderStageFlags, TConstArrayView<uint32> StaticUniformBuffers)
+{
+	check(StaticUniformBuffers.Num() <= FRHIShaderBindingLayout::MaxUniformBufferEntries);
+	if (StaticUniformBuffers.IsEmpty())
+	{
+		return;
+	}
+
+	// trim what we send
+	int32 StartIndex = 0;
+	for (; StartIndex < StaticUniformBuffers.Num(); ++StartIndex)
+	{
+		if (CurrentPushConstants[StartIndex] != StaticUniformBuffers[StartIndex])
+		{
+			break;
+		}
+	}
+	int32 EndIndex = StaticUniformBuffers.Num() - 1;
+	for (; EndIndex >= 0; --EndIndex)
+	{
+		if (CurrentPushConstants[EndIndex] != StaticUniformBuffers[EndIndex])
+		{
+			break;
+		}
+	}
+
+	if (StartIndex <= EndIndex)
+	{
+		FMemory::Memcpy(CurrentPushConstants.GetData(), StaticUniformBuffers.GetData(), StaticUniformBuffers.NumBytes());
+
+		const uint32 NumBytes = (EndIndex - StartIndex + 1) * sizeof(uint32);
+		if (Device.GetBindlessDescriptorManager()->UseDescriptorHeaps())
+		{
+			const VkPushDataInfoEXT PushDataInfo = {
+				.sType = VK_STRUCTURE_TYPE_PUSH_DATA_INFO_EXT,
+				.pNext = nullptr,
+				.offset = StartIndex * sizeof(uint32),
+				.data = {
+					.address = &CurrentPushConstants[StartIndex],
+					.size = NumBytes
+				}
+			};
+			VulkanRHI::vkCmdPushDataEXT(CommandBufferHandle, &PushDataInfo);
+		}
+		else
+		{
+			VulkanRHI::vkCmdPushConstants(
+				CommandBufferHandle, 
+				Device.GetBindlessDescriptorManager()->GetPipelineLayout(),
+				ShaderStageFlags, 
+				StartIndex * sizeof(uint32),
+				NumBytes,
+				&CurrentPushConstants[StartIndex]
+			);
+		}
+	}
+}
+
+void FVulkanCommandBuffer::SetBindingPushConstants(EShaderFrequency Frequency, uint64 GlobalsAddressOrFirstBinding, TConstArrayView<uint32> Bindings)
+{
+	checkf(Device.GetBindlessDescriptorManager()->UseDescriptorHeaps(), TEXT("Push constant bindings are only used by descriptor buffers"));
+	check(Bindings.Num() <= VulkanBindless::MaxUniformBuffersPerStage-1);
+
+	const uint32 ByteOffset = FVulkanBindlessDescriptorManager::GetBindingPushDataOffset(Frequency, 0);
+	const uint32 StartIndex = ByteOffset / sizeof(uint32);
+	if (GlobalsAddressOrFirstBinding != MAX_uint64)
+	{
+		FMemory::Memcpy(&CurrentPushConstants[StartIndex], &GlobalsAddressOrFirstBinding, sizeof(GlobalsAddressOrFirstBinding));
+	}
+
+	uint32 NumExtraValues = Bindings.Num();
+	if (Bindings.Num())
+	{
+		// trim what we send to the commandbuffer
+		const uint32 FirstExtraIndex = StartIndex + 2u; // Go past our initial 64bit value
+		for (; NumExtraValues > 0u; --NumExtraValues)
+		{
+			if (CurrentPushConstants[FirstExtraIndex + NumExtraValues - 1u] != Bindings[NumExtraValues - 1u])
+			{
+				break;
+			}
+		}
+
+		FMemory::Memcpy(&CurrentPushConstants[FirstExtraIndex], Bindings.GetData(), Bindings.NumBytes());
+	}
+
+	const VkPushDataInfoEXT PushDataInfo = {
+		.sType = VK_STRUCTURE_TYPE_PUSH_DATA_INFO_EXT,
+		.pNext = nullptr,
+		.offset = ByteOffset,
+		.data = {
+			.address = &CurrentPushConstants[StartIndex],
+			.size = (2u + NumExtraValues) * sizeof(uint32)
+		}
+	};
+	VulkanRHI::vkCmdPushDataEXT(CommandBufferHandle, &PushDataInfo);
+}
+
+
+
 
 FVulkanCommandBufferPool::FVulkanCommandBufferPool(FVulkanDevice& InDevice, FVulkanQueue& InQueue, EVulkanCommandBufferType InCommandBufferType)
 	: Device(InDevice)
diff --git a/Engine/Source/Runtime/VulkanRHI/Private/VulkanCommandBuffer.h b/Engine/Source/Runtime/VulkanRHI/Private/VulkanCommandBuffer.h
--- a/Engine/Source/Runtime/VulkanRHI/Private/VulkanCommandBuffer.h
+++ b/Engine/Source/Runtime/VulkanRHI/Private/VulkanCommandBuffer.h
@@ -171,6 +171,9 @@ public:
 	void BindGraphicsShaderObjects(TArrayView<FVulkanShader*> VulkanShaders);
 	void BindComputeShaderObject(FVulkanShader* VulkanShader);
 
+	void SetStaticUniformBufferPushConstants(VkShaderStageFlags ShaderStageFlags, TConstArrayView<uint32> StaticUniformBuffers);
+	void SetBindingPushConstants(EShaderFrequency Frequency, uint64 GlobalsAddressOrFirstBinding, TConstArrayView<uint32> Bindings);
+
 private:
 	VkCommandBuffer CommandBufferHandle = VK_NULL_HANDLE;
 	double SubmittedTime = 0.0f;
@@ -182,6 +185,9 @@ private:
 	VkPipeline CurrentPipelineHandle = VK_NULL_HANDLE;
 	VkShaderEXT CurrentComputeShaderObject = VK_NULL_HANDLE;
 	VkShaderEXT CurrentGraphicsShaderObjects[ShaderStage::NumGraphicsStages] = { VK_NULL_HANDLE };
+	TArray<uint32> CurrentPushConstants;
+
+	void ResetCurrentState();
 
 	TArray<VkEvent> EndedBarrierEvents;
 
diff --git a/Engine/Source/Runtime/VulkanRHI/Private/VulkanContext.cpp b/Engine/Source/Runtime/VulkanRHI/Private/VulkanContext.cpp
--- a/Engine/Source/Runtime/VulkanRHI/Private/VulkanContext.cpp
+++ b/Engine/Source/Runtime/VulkanRHI/Private/VulkanContext.cpp
@@ -706,7 +706,8 @@ void FVulkanCommandListContext::ApplyShaderBindingLayout(VkShaderStageFlags Shad
 	{
 		checkf(Device.GetBindlessDescriptorManager()->IsSupported(), TEXT("Static uniform buffer require bindless."));
 
-		uint32 UBBindlessIndices[FRHIShaderBindingLayout::MaxUniformBufferEntries];
+		TArray<uint32> UBBindlessIndices;
+		UBBindlessIndices.SetNum(ShaderBindingLayout->GetNumUniformBufferEntries());
 		for (uint32 Index = 0; Index < ShaderBindingLayout->GetNumUniformBufferEntries(); ++Index)
 		{
 			checkf(StaticUniformBuffers[Index], TEXT("Missing static uniform buffer for shader binding layout (%s)."),
@@ -719,25 +720,7 @@ void FVulkanCommandListContext::ApplyShaderBindingLayout(VkShaderStageFlags Shad
 			UBBindlessIndices[Index] = BindlessHandle.GetIndex();
 		}
 
-		const uint32 DataSize = sizeof(uint32) * ShaderBindingLayout->GetNumUniformBufferEntries();
-		if (Device.GetBindlessDescriptorManager()->UseDescriptorHeaps())
-		{
-			const VkPushDataInfoEXT PushDataInfo = {
-				.sType = VK_STRUCTURE_TYPE_PUSH_DATA_INFO_EXT,
-				.pNext = nullptr,
-				.offset = 0,
-				.data = {
-					.address = UBBindlessIndices,
-					.size = DataSize
-				}
-			};
-			VulkanRHI::vkCmdPushDataEXT(GetCommandBuffer().GetHandle(), &PushDataInfo);
-		}
-		else
-		{
-			VulkanRHI::vkCmdPushConstants(GetCommandBuffer().GetHandle(), Device.GetBindlessDescriptorManager()->GetPipelineLayout(),
-				ShaderStageFlags, 0 /*offset*/, DataSize, UBBindlessIndices);
-		}
+		GetCommandBuffer().SetStaticUniformBufferPushConstants(ShaderStageFlags, UBBindlessIndices);
 	}
 }
 
diff --git a/Engine/Source/Runtime/VulkanRHI/Private/VulkanPipeline.cpp b/Engine/Source/Runtime/VulkanRHI/Private/VulkanPipeline.cpp
--- a/Engine/Source/Runtime/VulkanRHI/Private/VulkanPipeline.cpp
+++ b/Engine/Source/Runtime/VulkanRHI/Private/VulkanPipeline.cpp
@@ -1536,7 +1536,7 @@ void FVulkanPipelineStateCacheManager::CreateEmptyFragmentStageLibrary()
 
 		if (Device->GetBindlessDescriptorManager()->UseDescriptorHeaps())
 		{
-			TConstArrayView<VkDescriptorSetAndBindingMappingEXT> BindingMappings = Device->GetBindlessDescriptorManager()->GetBindingMappings(SF_Pixel);
+			TConstArrayView<VkDescriptorSetAndBindingMappingEXT> BindingMappings = Device->GetBindlessDescriptorManager()->GetBindingMappings(SF_Pixel, false);
 			BindingMappingInfo = {
 				.sType = VK_STRUCTURE_TYPE_SHADER_DESCRIPTOR_SET_AND_BINDING_MAPPING_INFO_EXT,
 				.pNext = nullptr,
@@ -1805,9 +1805,9 @@ bool FVulkanPipelineStateCacheManager::CreateGfxPipelineFromEntry(FVulkanGraphic
 		Shaders[ShaderStage]->GetEntryPoint(EntryPoints[PipelineInfo.stageCount], 24);
 		ShaderStages[PipelineInfo.stageCount].pName = EntryPoints[PipelineInfo.stageCount];
 
+		const FVulkanShaderHeader& ShaderHeader = Shaders[ShaderStage]->GetCodeHeader();
 		if (Device->GetOptionalExtensions().HasEXTSubgroupSizeControl)
 		{
-			const FVulkanShaderHeader& ShaderHeader = Shaders[ShaderStage]->GetCodeHeader();
 			if (ShaderHeader.WaveSize > 0)
 			{
 				// Check if supported by this stage and Check if requested size is supported
@@ -1842,7 +1842,7 @@ bool FVulkanPipelineStateCacheManager::CreateGfxPipelineFromEntry(FVulkanGraphic
 
 		if (Device->GetBindlessDescriptorManager()->UseDescriptorHeaps())
 		{
-			TConstArrayView<VkDescriptorSetAndBindingMappingEXT> BindingMappings = Device->GetBindlessDescriptorManager()->GetBindingMappings(ShaderStage::GetFrequencyForGfxStage(CurrStage));
+			TConstArrayView<VkDescriptorSetAndBindingMappingEXT> BindingMappings = Device->GetBindlessDescriptorManager()->GetBindingMappings(ShaderStage::GetFrequencyForGfxStage(CurrStage), (ShaderHeader.PackedGlobalsSize > 0));
 			BindingMappingInfo[PipelineInfo.stageCount] = {
 				.sType = VK_STRUCTURE_TYPE_SHADER_DESCRIPTOR_SET_AND_BINDING_MAPPING_INFO_EXT,
 				.pNext = nullptr,
@@ -3032,7 +3032,7 @@ FVulkanComputePipeline* FVulkanPipelineStateCacheManager::CreateComputePipeline(
 	VkShaderDescriptorSetAndBindingMappingInfoEXT BindingMappingInfo;
 	if (Device->GetBindlessDescriptorManager()->UseDescriptorHeaps())
 	{
-		TConstArrayView<VkDescriptorSetAndBindingMappingEXT> BindingMappings = Device->GetBindlessDescriptorManager()->GetBindingMappings(SF_Compute);
+		TConstArrayView<VkDescriptorSetAndBindingMappingEXT> BindingMappings = Device->GetBindlessDescriptorManager()->GetBindingMappings(SF_Compute, (CSHeader.PackedGlobalsSize > 0));
 		BindingMappingInfo = {
 			.sType = VK_STRUCTURE_TYPE_SHADER_DESCRIPTOR_SET_AND_BINDING_MAPPING_INFO_EXT,
 			.pNext = nullptr,
diff --git a/Engine/Source/Runtime/VulkanRHI/Private/VulkanPipelineState.cpp b/Engine/Source/Runtime/VulkanRHI/Private/VulkanPipelineState.cpp
--- a/Engine/Source/Runtime/VulkanRHI/Private/VulkanPipelineState.cpp
+++ b/Engine/Source/Runtime/VulkanRHI/Private/VulkanPipelineState.cpp
@@ -342,30 +342,35 @@ void FVulkanComputePipelineDescriptorState::UpdateBindlessDescriptors(FVulkanCom
 
 	if (Device.GetBindlessDescriptorManager()->UseDescriptorHeaps())
 	{
-		TArray<uint32> PushData = DSWriter[ShaderStage::EStage::Compute].BindlessUBIndices;
-		check(PushData.Num() == Header.NumBoundUniformBuffers);
+		uint64 GlobalAddressOrFirstBinding = DSWriter[ShaderStage::EStage::Compute].BindlessUBIndices[0];
+		if (Header.PackedGlobalsSize > 0)
+		{
+			check(PackedUniformBufferState.Mask <= 1);
+			if (PackedUniformBufferState.Mask != 0)
+			{
+				const FPackedUniformBuffers::FPackedBuffer& StagedUniformBuffer = PackedUniformBufferState.UniformBuffers.GetBuffer();
+				const int32 UBSize = StagedUniformBuffer.Num();
+
+				VkDescriptorAddressInfoEXT AddressInfo;
+				VulkanRHI::FVulkanAllocation TempAllocation;
+				uint8* MappedPointer = Device.GetTempBlockAllocator().Alloc(UBSize, UBOffsetAlignment, Context, TempAllocation, &AddressInfo);
+				FMemory::Memcpy(MappedPointer, StagedUniformBuffer.GetData(), UBSize);
+
+				GlobalAddressOrFirstBinding = AddressInfo.address;
+			}
+			else
+			{
+				GlobalAddressOrFirstBinding = MAX_uint64;
+			}
+		}
 
-		VulkanRHI::FVulkanAllocation TempAllocation = GetGlobalUBAlloc();
-		if (TempAllocation.IsValid())
+		TConstArrayView<uint32> Bindings;
+		if (Header.NumBoundUniformBuffers > 1)
 		{
-			// NOTE: Can be replaced with BufferAddr straight in push data when DescriptorBuffer is removes (instead of the temporary handle)
-			FVulkanBindlessDescriptorManager* Mgr = Device.GetBindlessDescriptorManager();
-			FRHIDescriptorHandle BindlessHandle = Mgr->AllocateDescriptor(ERHIDescriptorType::CBV);
-			Mgr->UpdateDescriptor({}, BindlessHandle, TempAllocation.GetBufferHandle(), TempAllocation.Offset, TempAllocation.Size);
-			PushData[0] = BindlessHandle.GetIndex();
-			Device.GetDeferredDeletionQueue().EnqueueBindlessHandle(BindlessHandle);
+			Bindings = TConstArrayView<uint32>(&DSWriter[ShaderStage::EStage::Compute].BindlessUBIndices[1], Header.NumBoundUniformBuffers - 1);
 		}
 
-		const VkPushDataInfoEXT PushDataInfo = {
-			.sType = VK_STRUCTURE_TYPE_PUSH_DATA_INFO_EXT,
-			.pNext = nullptr,
-			.offset = (FRHIShaderBindingLayout::MaxUniformBufferEntries * sizeof(uint32)),  // place them after the static UBs
-			.data = {
-				.address = PushData.GetData(),
-				.size = PushData.NumBytes()
-			}
-		};
-		VulkanRHI::vkCmdPushDataEXT(Context.GetCommandBuffer().GetHandle(), &PushDataInfo);
+		Context.GetCommandBuffer().SetBindingPushConstants(SF_Compute, GlobalAddressOrFirstBinding, Bindings);
 	}
 	else
 	{
@@ -642,10 +647,6 @@ void FVulkanGraphicsPipelineDescriptorState::UpdateBindlessDescriptors(FVulkanCo
 
 	if (Device.GetBindlessDescriptorManager()->UseDescriptorHeaps())
 	{
-		TStaticArray<uint32, VulkanBindless::MaxUniformBuffersTotal> PushData;
-		FMemory::Memzero(PushData);
-
-		int32 PushDataCount = 0;
 		for (int32 Stage = 0; Stage < ShaderStage::NumGraphicsStages; ++Stage)
 		{
 			const EShaderFrequency Frequency = GetFrequencyForGfxStage((ShaderStage::EStage)Stage);
@@ -664,48 +665,35 @@ void FVulkanGraphicsPipelineDescriptorState::UpdateBindlessDescriptors(FVulkanCo
 				continue;
 			}
 
-			// NOTE: All this shuffling around to fit everything in MaxUniformBuffersTotal won't be necessary when DescriptorBuffer support is removed
-			const int32 BindingOffset = VulkanBindless::GetOffsetForFrequency(Frequency);
-			if (BindingOffset < 0)
+			uint64 GlobalAddressOrFirstBinding = DSWriter[Stage].BindlessUBIndices[0];
+			if (Header.PackedGlobalsSize > 0)
 			{
-				for (int32 UBIndex = 0; UBIndex < BindlessUBIndices.Num(); ++UBIndex)
+				check(PackedUniformBufferStates[Stage].Mask <= 1);
+				if (PackedUniformBufferStates[Stage].Mask != 0)
 				{
-					const int32 PushDataIndex = VulkanBindless::MaxUniformBuffersTotal - UBIndex - 1;
-					PushData[PushDataIndex] = BindlessUBIndices[UBIndex];
-				}
+					const FPackedUniformBuffers::FPackedBuffer& StagedUniformBuffer = PackedUniformBufferStates[Stage].UniformBuffers.GetBuffer();
+					const int32 UBSize = StagedUniformBuffer.Num();
 
-				PushDataCount = VulkanBindless::MaxUniformBuffersTotal;
-			}
-			else
-			{
-				FMemory::Memcpy(&PushData[BindingOffset], BindlessUBIndices.GetData(), BindlessUBIndices.NumBytes());
-				PushDataCount = FMath::Max(PushDataCount, BindingOffset + BindlessUBIndices.Num());
+					VkDescriptorAddressInfoEXT AddressInfo;
+					VulkanRHI::FVulkanAllocation TempAllocation;
+					uint8* MappedPointer = Device.GetTempBlockAllocator().Alloc(UBSize, UBOffsetAlignment, Context, TempAllocation, &AddressInfo);
+					FMemory::Memcpy(MappedPointer, StagedUniformBuffer.GetData(), UBSize);
+
+					GlobalAddressOrFirstBinding = AddressInfo.address;
+				}
+				else
+				{
+					GlobalAddressOrFirstBinding = MAX_uint64;
+				}
 			}
 
-			VulkanRHI::FVulkanAllocation TempAllocation = GetGlobalUBAlloc(Stage);
-			if (TempAllocation.IsValid())
+			TConstArrayView<uint32> Bindings;
+			if (Header.NumBoundUniformBuffers > 1)
 			{
-				// NOTE: Can be replaced with BufferAddr straight in push data when DescriptorBuffer is removes
-				FVulkanBindlessDescriptorManager* Mgr = Device.GetBindlessDescriptorManager();
-				FRHIDescriptorHandle BindlessHandle = Mgr->AllocateDescriptor(ERHIDescriptorType::CBV);
-				Mgr->UpdateDescriptor({}, BindlessHandle, TempAllocation.GetBufferHandle(), TempAllocation.Offset, TempAllocation.Size);
-				PushData[(BindingOffset < 0) ? VulkanBindless::MaxUniformBuffersTotal - 1 : BindingOffset] = BindlessHandle.GetIndex();
-				Device.GetDeferredDeletionQueue().EnqueueBindlessHandle(BindlessHandle);
+				Bindings = TConstArrayView<uint32>(&BindlessUBIndices[1], Header.NumBoundUniformBuffers - 1);
 			}
-		}
 
-		if (PushDataCount)
-		{
-			const VkPushDataInfoEXT PushDataInfo = {
-				.sType = VK_STRUCTURE_TYPE_PUSH_DATA_INFO_EXT,
-				.pNext = nullptr,
-				.offset = (FRHIShaderBindingLayout::MaxUniformBufferEntries * sizeof(uint32)),  // place them after the static UBs
-				.data = {
-					.address = PushData.GetData(),
-					.size = PushDataCount * sizeof(uint32)
-				}
-			};
-			VulkanRHI::vkCmdPushDataEXT(Context.GetCommandBuffer().GetHandle(), &PushDataInfo);
+			Context.GetCommandBuffer().SetBindingPushConstants(Frequency, GlobalAddressOrFirstBinding, Bindings);
 		}
 	}
 	else
diff --git a/Engine/Source/Runtime/VulkanRHI/Private/VulkanRayTracing.cpp b/Engine/Source/Runtime/VulkanRHI/Private/VulkanRayTracing.cpp
--- a/Engine/Source/Runtime/VulkanRHI/Private/VulkanRayTracing.cpp
+++ b/Engine/Source/Runtime/VulkanRHI/Private/VulkanRayTracing.cpp
@@ -1482,7 +1482,7 @@ public:
 				TStaticArray<VkShaderDescriptorSetAndBindingMappingInfoEXT, MaxHitGroupShaderCount> BindingMappingInfos;
 				if (Device.GetBindlessDescriptorManager()->UseDescriptorHeaps())
 				{
-					TConstArrayView<VkDescriptorSetAndBindingMappingEXT> BindingMappings = Device.GetBindlessDescriptorManager()->GetBindingMappings(Frequency);
+					TConstArrayView<VkDescriptorSetAndBindingMappingEXT> BindingMappings = Device.GetBindlessDescriptorManager()->GetBindingMappings(Frequency, false);
 					for (uint32 StageIndex=0; StageIndex < ShaderCount; ++StageIndex)
 					{
 						BindingMappingInfos[StageIndex] = {
@@ -1972,30 +1972,28 @@ static void SetRayGenResources(FVulkanDevice& Device, FVulkanCommandListContext&
 	if (Device.GetBindlessDescriptorManager()->UseDescriptorHeaps())
 	{
 		TArray<uint32> PushData;
-		PushData.SetNumZeroed(MaxUniformBuffers);
+		PushData.Reserve(MaxUniformBuffers);
 
-		uint32 MaxIndex = 0;
+		uint64 GlobalAddressOrFirstBinding = MAX_uint64;
 		for (uint32 UBIndex = 0; UBIndex < MaxUniformBuffers; ++UBIndex)
 		{
 			const FVulkanUniformBuffer* UniformBuffer = ResourceCast(InGlobalResourceBindings.UniformBuffers[UBIndex]);
 			if (UniformBuffer)
 			{
 				check(UniformBuffer->GetBindlessHandle().IsValid());
-				PushData[UBIndex] = UniformBuffer->GetBindlessHandle().GetIndex();
-				MaxIndex = UBIndex;
+				if (UBIndex == 0)
+				{
+					GlobalAddressOrFirstBinding = UniformBuffer->GetBindlessHandle().GetIndex();
+				}
+				else
+				{
+					PushData.SetNumZeroed(UBIndex);
+					PushData[UBIndex-1] = UniformBuffer->GetBindlessHandle().GetIndex();
+				}
 			}
 		}
 
-		const VkPushDataInfoEXT PushDataInfo = {
-			.sType = VK_STRUCTURE_TYPE_PUSH_DATA_INFO_EXT,
-			.pNext = nullptr,
-			.offset = (FRHIShaderBindingLayout::MaxUniformBufferEntries * sizeof(uint32)),  // place them after the static UBs
-			.data = {
-				.address = PushData.GetData(),
-				.size = (MaxIndex + 1) * sizeof(uint32)
-			}
-		};
-		VulkanRHI::vkCmdPushDataEXT(Context.GetCommandBuffer().GetHandle(), &PushDataInfo);
+		Context.GetCommandBuffer().SetBindingPushConstants(SF_RayGen, GlobalAddressOrFirstBinding, PushData);
 	}
 	else 
 	{
