VulkanRHI: replace some UE_LOG by UE_LOGF

#rnx

#Synced-CL 52393244

diff --git a/Engine/Source/Runtime/VulkanRHI/Private/VulkanBindlessDescriptorManager.cpp b/Engine/Source/Runtime/VulkanRHI/Private/VulkanBindlessDescriptorManager.cpp
--- a/Engine/Source/Runtime/VulkanRHI/Private/VulkanBindlessDescriptorManager.cpp
+++ b/Engine/Source/Runtime/VulkanRHI/Private/VulkanBindlessDescriptorManager.cpp
@@ -209,17 +209,17 @@ static FVulkanBindlessDescriptorManager::EBindlessType VerifyDescriptorHeapSuppo
 
 		if (bMeetsPropertiesRequirements)
 		{
-			UE_LOG(LogRHI, Display, TEXT("Bindless will use VK_EXT_descriptor_heap extension."));
+			UE_LOGF(LogRHI, Display, "Bindless will use VK_EXT_descriptor_heap extension.");
 			return FVulkanBindlessDescriptorManager::EBindlessType::DescriptorHeap;
 		}
 		else
 		{
-			UE_LOG(LogRHI, Warning, TEXT("Bindless descriptor were requested but NOT enabled because of insufficient property support."));
+			UE_LOGF(LogRHI, Warning, "Bindless descriptor were requested but NOT enabled because of insufficient property support.");
 		}
 	}
 	else
 	{
-		UE_LOG(LogRHI, Warning, TEXT("Bindless descriptor were requested but NOT enabled because of missing extension support."));
+		UE_LOGF(LogRHI, Warning, "Bindless descriptor were requested but NOT enabled because of missing extension support.");
 	}
 
 	return FVulkanBindlessDescriptorManager::EBindlessType::Disabled;
@@ -255,17 +255,17 @@ static FVulkanBindlessDescriptorManager::EBindlessType VerifyDescriptorBufferSup
 
 		if (bMeetsPropertiesRequirements)
 		{
-			UE_LOG(LogRHI, Display, TEXT("Bindless will use VK_EXT_descriptor_buffer extension."));
+			UE_LOGF(LogRHI, Display, "Bindless will use VK_EXT_descriptor_buffer extension.");
 			return FVulkanBindlessDescriptorManager::EBindlessType::DescriptorBuffer;
 		}
 		else
 		{
-			UE_LOG(LogRHI, Warning, TEXT("Bindless descriptor were requested but NOT enabled because of insufficient property support."));
+			UE_LOGF(LogRHI, Warning, "Bindless descriptor were requested but NOT enabled because of insufficient property support.");
 		}
 	}
 	else
 	{
-		UE_LOG(LogRHI, Warning, TEXT("Bindless descriptor were requested but NOT enabled because of missing extension support."));
+		UE_LOGF(LogRHI, Warning, "Bindless descriptor were requested but NOT enabled because of missing extension support.");
 	}
 
 	return FVulkanBindlessDescriptorManager::EBindlessType::Disabled;
@@ -289,7 +289,7 @@ static FVulkanBindlessDescriptorManager::EBindlessType VerifySupport(FVulkanDevi
 		}
 		else
 		{
-			UE_LOG(LogRHI, Warning, TEXT("Bindless in Vulkan must currently be fully enabled (all samplers and resources) or fully disabled."));
+			UE_LOGF(LogRHI, Warning, "Bindless in Vulkan must currently be fully enabled (all samplers and resources) or fully disabled.");
 		}
 	}
 
@@ -298,13 +298,13 @@ static FVulkanBindlessDescriptorManager::EBindlessType VerifySupport(FVulkanDevi
 		extern TAutoConsoleVariable<int32> GDynamicGlobalUBs;
 		if (GDynamicGlobalUBs->GetInt() != 0)
 		{
-			UE_LOG(LogRHI, Warning, TEXT("Dynamic Uniform Buffers are enabled, but they will not be used with Vulkan bindless."));
+			UE_LOGF(LogRHI, Warning, "Dynamic Uniform Buffers are enabled, but they will not be used with Vulkan bindless.");
 		}
 
 		extern int32 GVulkanEnableDefrag;
 		if (GVulkanEnableDefrag != 0)  // :todo-jn: to be turned back on with new defragger
 		{
-			UE_LOG(LogRHI, Warning, TEXT("Memory defrag is enabled, but it will not be used with Vulkan bindless."));
+			UE_LOGF(LogRHI, Warning, "Memory defrag is enabled, but it will not be used with Vulkan bindless.");
 			GVulkanEnableDefrag = 0;
 		}
 	}
diff --git a/Engine/Source/Runtime/VulkanRHI/Private/VulkanPipeline.cpp b/Engine/Source/Runtime/VulkanRHI/Private/VulkanPipeline.cpp
--- a/Engine/Source/Runtime/VulkanRHI/Private/VulkanPipeline.cpp
+++ b/Engine/Source/Runtime/VulkanRHI/Private/VulkanPipeline.cpp
@@ -420,17 +420,17 @@ bool FVulkanPipelineStateCacheManager::Load(const TArray<FString>& CacheFilename
 				}
 
 				double EndTime = FPlatformTime::Seconds();
-				UE_LOG(LogVulkanRHI, Display, TEXT("FVulkanPipelineStateCacheManager: Loaded binary pipeline cache %s in %.3f seconds"), *BinaryCacheFilename, (float)(EndTime - BeginTime));
+				UE_LOGF(LogVulkanRHI, Display, "FVulkanPipelineStateCacheManager: Loaded binary pipeline cache %ls in %.3f seconds", *BinaryCacheFilename, (float)(EndTime - BeginTime));
 				bResult = true;
 			}
 			else
 			{
-				UE_LOG(LogVulkanRHI, Error, TEXT("FVulkanPipelineStateCacheManager: Mismatched binary pipeline cache %s"), *BinaryCacheFilename);
+				UE_LOGF(LogVulkanRHI, Error, "FVulkanPipelineStateCacheManager: Mismatched binary pipeline cache %ls", *BinaryCacheFilename);
 			}
 		}
 		else
 		{
-			UE_LOG(LogVulkanRHI, Display, TEXT("FVulkanPipelineStateCacheManager: Binary pipeline cache '%s' not found."), *BinaryCacheFilename);
+			UE_LOGF(LogVulkanRHI, Display, "FVulkanPipelineStateCacheManager: Binary pipeline cache '%ls' not found.", *BinaryCacheFilename);
 		}
 	}
 
@@ -453,7 +453,7 @@ bool FVulkanPipelineStateCacheManager::Load(const TArray<FString>& CacheFilename
 				bool Valid = File.Load(Ar);
 				if (!Valid)
 				{
-					UE_LOG(LogVulkanRHI, Warning, TEXT("Unable to load lru pipeline cache '%s'"), *LruCacheFilename);
+					UE_LOGF(LogVulkanRHI, Warning, "Unable to load lru pipeline cache '%ls'", *LruCacheFilename);
 					bResult = false;
 				}
 
@@ -461,11 +461,11 @@ bool FVulkanPipelineStateCacheManager::Load(const TArray<FString>& CacheFilename
 				{
 					LRU2SizeList.Add(File.PipelineSizes[Index].ShaderHash, File.PipelineSizes[Index]);
 				}
-				UE_LOG(LogVulkanRHI, Display, TEXT("Loaded %d LRU size entries for '%s'"), File.PipelineSizes.Num(), *LruCacheFilename);
+				UE_LOGF(LogVulkanRHI, Display, "Loaded %d LRU size entries for '%ls'", File.PipelineSizes.Num(), *LruCacheFilename);
 			}
 			else
 			{
-				UE_LOG(LogVulkanRHI, Warning, TEXT("Unable to load lru pipeline cache '%s'"), *LruCacheFilename);
+				UE_LOGF(LogVulkanRHI, Warning, "Unable to load lru pipeline cache '%ls'", *LruCacheFilename);
 				bResult = false;
 			}
 		}
@@ -491,7 +491,7 @@ void FVulkanPipelineStateCacheManager::InitAndLoad(const TArray<FString>& CacheF
 	if (FParse::Param(FCommandLine::Get(), TEXT("ClearVulkanBinaryProgramCache")) || 
 		FParse::Param(FCommandLine::Get(), TEXT("clearPSODriverCache")))
 	{
-		UE_LOG(LogVulkanRHI, Log, TEXT("FVulkanPipelineStateCacheManager: Clearing pipeline cache on disk"));
+		UE_LOGF(LogVulkanRHI, Log, "FVulkanPipelineStateCacheManager: Clearing pipeline cache on disk");
 		for (const FString& CacheFilename : CacheFilenames)
 		{
 			FString BinaryCacheFilename = FVulkanPlatform::CreatePSOBinaryCacheFilename(Device, CacheFilename);
@@ -502,7 +502,7 @@ void FVulkanPipelineStateCacheManager::InitAndLoad(const TArray<FString>& CacheF
 	
 	if (GEnablePipelineCacheLoadCvar.GetValueOnAnyThread() == 0)
 	{
-		UE_LOG(LogVulkanRHI, Display, TEXT("Not loading pipeline cache per r.Vulkan.PipelineCacheLoad=0"));
+		UE_LOGF(LogVulkanRHI, Display, "Not loading pipeline cache per r.Vulkan.PipelineCacheLoad=0");
 	}
 	else
 	{
@@ -558,13 +558,13 @@ void FVulkanPipelineStateCacheManager::OnShaderPipelineCacheOpened(FString const
 	//TODO: support reloading the same cache
 	if (CompiledPSOCaches.Contains(VersionGuid))
 	{
-		UE_LOG(LogVulkanRHI, Log, TEXT("FVulkanPipelineStateCacheManager::OnShaderPipelineCacheOpened attempts to load a cache that was already loaded before %s %s"), *Name, *VersionGuid.ToString());
+		UE_LOGF(LogVulkanRHI, Log, "FVulkanPipelineStateCacheManager::OnShaderPipelineCacheOpened attempts to load a cache that was already loaded before %ls %ls", *Name, *VersionGuid.ToString());
 		return;
 	}
 
 	CurrentPrecompilingPSOCacheGuid = VersionGuid;
 
-	UE_LOG(LogVulkanRHI, Log, TEXT("FVulkanPipelineStateCacheManager::OnShaderPipelineCacheOpened %s %d %s"), *Name, Count, *VersionGuid.ToString());
+	UE_LOGF(LogVulkanRHI, Log, "FVulkanPipelineStateCacheManager::OnShaderPipelineCacheOpened %ls %d %ls", *Name, Count, *VersionGuid.ToString());
 
 	IPlatformFile& PlatformFile = FPlatformFileManager::Get().GetPlatformFile();
 
@@ -605,13 +605,13 @@ void FVulkanPipelineStateCacheManager::OnShaderPipelineCacheOpened(FString const
 		}
 		else
 		{
-			UE_LOG(LogVulkanRHI, Log, TEXT("FVulkanPipelineStateCacheManager: PSO cache failed to load, deleting file: %s"), *CompiledPSOCacheFolderName);
+			UE_LOGF(LogVulkanRHI, Log, "FVulkanPipelineStateCacheManager: PSO cache failed to load, deleting file: %ls", *CompiledPSOCacheFolderName);
 			FPlatformFileManager::Get().GetPlatformFile().DeleteFile(*CompiledPSOCacheFolderName);
 		}
 	}
 	else
 	{
-		UE_LOG(LogVulkanRHI, Log, TEXT("FVulkanPipelineStateCacheManager: %s does not exist."), *CompiledPSOCacheFolderName);
+		UE_LOGF(LogVulkanRHI, Log, "FVulkanPipelineStateCacheManager: %ls does not exist.", *CompiledPSOCacheFolderName);
 	}
 
 
@@ -630,7 +630,7 @@ void FVulkanPipelineStateCacheManager::OnShaderPipelineCacheOpened(FString const
 
 void FVulkanPipelineStateCacheManager::OnShaderPipelineCachePrecompilationComplete(uint32 Count, double Seconds, const FShaderPipelineCache::FShaderCachePrecompileContext& ShaderCachePrecompileContext)
 {
-	UE_LOG(LogVulkanRHI, Log, TEXT("FVulkanPipelineStateCacheManager::OnShaderPipelineCachePrecompilationComplete"));
+	UE_LOGF(LogVulkanRHI, Log, "FVulkanPipelineStateCacheManager::OnShaderPipelineCachePrecompilationComplete");
 
 #if PLATFORM_ANDROID
 	if (FVulkanAndroidPlatform::AreRemoteCompileServicesActive())
@@ -686,16 +686,16 @@ void FVulkanPipelineStateCacheManager::SavePSOCache(const FString& CacheFilename
 
 			if (FFileHelper::SaveArrayToFile(DeviceCache, *BinaryCacheFilename))
 			{
-				UE_LOG(LogVulkanRHI, Display, TEXT("FVulkanPipelineStateCacheManager: Saved device pipeline cache file '%s', %d bytes"), *BinaryCacheFilename, DeviceCache.Num());
+				UE_LOGF(LogVulkanRHI, Display, "FVulkanPipelineStateCacheManager: Saved device pipeline cache file '%ls', %d bytes", *BinaryCacheFilename, DeviceCache.Num());
 			}
 			else
 			{
-				UE_LOG(LogVulkanRHI, Error, TEXT("FVulkanPipelineStateCacheManager: Failed to save device pipeline cache file '%s', %d bytes"), *BinaryCacheFilename, DeviceCache.Num());
+				UE_LOGF(LogVulkanRHI, Error, "FVulkanPipelineStateCacheManager: Failed to save device pipeline cache file '%ls', %d bytes", *BinaryCacheFilename, DeviceCache.Num());
 			}
 		}
 		else if (Result == VK_INCOMPLETE || Result == VK_ERROR_OUT_OF_HOST_MEMORY)
 		{
-			UE_LOG(LogVulkanRHI, Warning, TEXT("Failed to get Vulkan pipeline cache data. Error %d, %d bytes"), Result, Size);
+			UE_LOGF(LogVulkanRHI, Warning, "Failed to get Vulkan pipeline cache data. Error %d, %d bytes", Result, Size);
 
 			//TODO: Resave it when we shutdown the manager?!
 			VulkanRHI::vkDestroyPipelineCache(Device->GetHandle(), PipelineCacheExclusive.Get(), VULKAN_CPU_ALLOCATOR);
@@ -725,11 +725,11 @@ void FVulkanPipelineStateCacheManager::SavePSOCache(const FString& CacheFilename
 
 		if (FFileHelper::SaveArrayToFile(MemFile, *LruCacheFilename))
 		{
-			UE_LOG(LogVulkanRHI, Display, TEXT("FVulkanPipelineStateCacheManager: Saved pipeline lru pipeline cache file '%s', %d hashes, %d bytes"), *LruCacheFilename, LRU2SizeList.Num(), MemFile.Num());
+			UE_LOGF(LogVulkanRHI, Display, "FVulkanPipelineStateCacheManager: Saved pipeline lru pipeline cache file '%ls', %d hashes, %d bytes", *LruCacheFilename, LRU2SizeList.Num(), MemFile.Num());
 		}
 		else
 		{
-			UE_LOG(LogVulkanRHI, Error, TEXT("FVulkanPipelineStateCacheManager: Failed to save pipeline lru pipeline cache file '%s', %d hashes, %d bytes"), *LruCacheFilename, LRU2SizeList.Num(), MemFile.Num());
+			UE_LOGF(LogVulkanRHI, Error, "FVulkanPipelineStateCacheManager: Failed to save pipeline lru pipeline cache file '%ls', %d hashes, %d bytes", *LruCacheFilename, LRU2SizeList.Num(), MemFile.Num());
 		}
 	}
 }
@@ -1963,7 +1963,7 @@ bool FVulkanPipelineStateCacheManager::CreateGfxPipelineFromEntry(FVulkanGraphic
 	{
 		FString ShaderHashes = GfxShaderHashesToString(Shaders);
 
-		UE_LOG(LogVulkanRHI, Error, TEXT("Failed to create graphics pipeline.\nShaders in pipeline: %s"), *ShaderHashes);
+		UE_LOGF(LogVulkanRHI, Error, "Failed to create graphics pipeline.\nShaders in pipeline: %ls", *ShaderHashes);
 		return false;
 	}
 
@@ -1971,7 +1971,7 @@ bool FVulkanPipelineStateCacheManager::CreateGfxPipelineFromEntry(FVulkanGraphic
 	double Delta = EndTime - BeginTime;
 	if (Delta > HitchTime)
 	{
-		UE_LOG(LogVulkanRHI, Verbose, TEXT("Hitchy gfx pipeline key CS (%.3f ms)"), (float)(Delta * 1000.0));
+		UE_LOGF(LogVulkanRHI, Verbose, "Hitchy gfx pipeline key CS (%.3f ms)", (float)(Delta * 1000.0));
 	}
 
 	INC_DWORD_STAT(STAT_VulkanNumPSOs);
@@ -2251,7 +2251,7 @@ VkResult FVulkanPipelineStateCacheManager::CreateVKPipeline(FVulkanGraphicsPipel
 
 				if (!Diff)
 				{
-					UE_LOG(LogVulkanRHI, Warning, TEXT("Shader size was computed as zero, using 20k instead."));
+					UE_LOGF(LogVulkanRHI, Warning, "Shader size was computed as zero, using 20k instead.");
 					Diff = 20 * 1024;
 				}
 				FVulkanPipelineSize PipelineSize;
@@ -2829,7 +2829,7 @@ FGraphicsPipelineStateRHIRef FVulkanPipelineStateCacheManager::RHICreateGraphics
 				double Delta = EndTime - BeginTime;
 				if (Delta > HitchTime)
 				{
-					UE_LOG(LogVulkanRHI, Verbose, TEXT("Hitchy gfx pipeline (%.3f ms)"), (float)(Delta * 1000.0));
+					UE_LOGF(LogVulkanRHI, Verbose, "Hitchy gfx pipeline (%.3f ms)", (float)(Delta * 1000.0));
 				}
 			}
 			FScopeLock Lock(&GraphicsPSOLockedCS); 
@@ -2921,7 +2921,7 @@ FVulkanComputePipeline* FVulkanPipelineStateCacheManager::GetOrCreateComputePipe
 	double Delta = EndTime - BeginTime;
 	if (Delta > HitchTime)
 	{
-		UE_LOG(LogVulkanRHI, Verbose, TEXT("Hitchy compute pipeline key CS (%.3f ms)"), (float)(Delta * 1000.0));
+		UE_LOGF(LogVulkanRHI, Verbose, "Hitchy compute pipeline key CS (%.3f ms)", (float)(Delta * 1000.0));
 	}
 
 	{
@@ -3054,7 +3054,7 @@ FVulkanComputePipeline* FVulkanPipelineStateCacheManager::CreateComputePipeline(
 	if (Result != VK_SUCCESS)
 	{
 		FString ComputeHash = Shader->GetHash().ToString();
-		UE_LOG(LogVulkanRHI, Error, TEXT("Failed to create compute pipeline.\nShaders in pipeline: CS: %s"), *ComputeHash);
+		UE_LOGF(LogVulkanRHI, Error, "Failed to create compute pipeline.\nShaders in pipeline: CS: %ls", *ComputeHash);
 		Pipeline->SetValid(false);
 	}
 
@@ -3121,20 +3121,20 @@ bool FVulkanPipelineStateCacheManager::FVulkanLRUCacheFile::Load(FArchive& Ar)
 	Ar << Header.Version;
 	if (Header.Version != LRU_CACHE_VERSION)
 	{
-		UE_LOG(LogVulkanRHI, Warning, TEXT("Unable to load lru pipeline cache due to mismatched Version %d != %d"), Header.Version, (int32)LRU_CACHE_VERSION);
+		UE_LOGF(LogVulkanRHI, Warning, "Unable to load lru pipeline cache due to mismatched Version %d != %d", Header.Version, (int32)LRU_CACHE_VERSION);
 		return false;
 	}
 
 	Ar << Header.SizeOfPipelineSizes;
 	if (Header.SizeOfPipelineSizes != (int32)(sizeof(FVulkanPipelineSize)))
 	{
-		UE_LOG(LogVulkanRHI, Warning, TEXT("Unable to load lru pipeline cache due to mismatched size of FVulkanPipelineSize %d != %d; forgot to bump up LRU_CACHE_VERSION?"), Header.SizeOfPipelineSizes, (int32)sizeof(FVulkanPipelineSize));
+		UE_LOGF(LogVulkanRHI, Warning, "Unable to load lru pipeline cache due to mismatched size of FVulkanPipelineSize %d != %d; forgot to bump up LRU_CACHE_VERSION?", Header.SizeOfPipelineSizes, (int32)sizeof(FVulkanPipelineSize));
 		return false;
 	}
 
 	if (!SerializeArray(Ar, PipelineSizes))
 	{
-		UE_LOG(LogVulkanRHI, Warning, TEXT("Unable to load lru pipeline cache due to invalid archive data!"));
+		UE_LOGF(LogVulkanRHI, Warning, "Unable to load lru pipeline cache due to invalid archive data!");
 		return false;
 	}
 
@@ -3334,7 +3334,7 @@ void FVulkanPipelineStateCacheManager::LRUTouch(FVulkanGraphicsPipelineState* PS
 			double Delta = EndTime - BeginTime;
 			if (Delta > HitchTime)
 			{
-				UE_LOG(LogVulkanRHI, Verbose, TEXT("Hitchy gfx pipeline (%.3f ms)"), (float)(Delta * 1000.0));
+				UE_LOGF(LogVulkanRHI, Verbose, "Hitchy gfx pipeline (%.3f ms)", (float)(Delta * 1000.0));
 			}
 
 			if(bUseLRU)
diff --git a/Engine/Source/Runtime/VulkanRHI/Private/VulkanRayTracing.cpp b/Engine/Source/Runtime/VulkanRHI/Private/VulkanRayTracing.cpp
--- a/Engine/Source/Runtime/VulkanRHI/Private/VulkanRayTracing.cpp
+++ b/Engine/Source/Runtime/VulkanRHI/Private/VulkanRayTracing.cpp
@@ -1541,7 +1541,7 @@ void FVulkanDevice::InitializeRayTracing()
 		const VkPhysicalDeviceRayTracingPipelinePropertiesKHR& RayTracingPipelineProps = GetOptionalExtensionProperties().RayTracingPipelineProps;
 		if ((uint32)GVulkanRayTracingMaxShaderGroupStride > RayTracingPipelineProps.maxShaderGroupStride)
 		{
-			UE_LOG(LogRHI, Warning, TEXT("Specified value for r.Vulkan.RayTracing.MaxShaderGroupStride is too large for this device! It will be capped."));
+			UE_LOGF(LogRHI, Warning, "Specified value for r.Vulkan.RayTracing.MaxShaderGroupStride is too large for this device! It will be capped.");
 		}
 		GVulkanRayTracingMaxShaderGroupStride = FMath::Min<VkDeviceSize>(RayTracingPipelineProps.maxShaderGroupStride, GVulkanRayTracingMaxShaderGroupStride);
 
@@ -1574,7 +1574,7 @@ static FVulkanUniformBuffer* GetDummyUB(FVulkanDevice& Device, uint32 UBLayoutHa
 		FVulkanUniformBuffer* DummyUB = new FVulkanUniformBuffer(Device, DummyLayout, DummyContent.GetData(), UniformBuffer_MultiFrame, EUniformBufferValidation::None);
 		DummyUBs.Add(UBLayoutHash, DummyUB);
 		const FString& LayoutName = DummyLayout->GetDebugName();
-		UE_LOG(LogRHI, Warning, TEXT("Vulkan ray tracing using DummyUB for %s."), LayoutName.IsEmpty() ? TEXT("<unknown>") : *LayoutName);
+		UE_LOGF(LogRHI, Warning, "Vulkan ray tracing using DummyUB for %ls.", LayoutName.IsEmpty() ? TEXT("<unknown>") : *LayoutName);
 		return DummyUB;
 	}
 	return nullptr;
