VulkanRHI: Remove compute stage from GVulkanDevicePipelineStageBits and rename it GVulkanDeviceGraphicsPipelineStageBits (instead of manually removing compute everywhere it's used).

#rnx

#Synced-CL 52444545

diff --git a/Engine/Source/Runtime/VulkanRHI/Private/VulkanBarriers.cpp b/Engine/Source/Runtime/VulkanRHI/Private/VulkanBarriers.cpp
--- a/Engine/Source/Runtime/VulkanRHI/Private/VulkanBarriers.cpp
+++ b/Engine/Source/Runtime/VulkanRHI/Private/VulkanBarriers.cpp
@@ -7,8 +7,8 @@
 #include "ProfilingDebugging/RealtimeGPUProfiler.h"
 #include "RHICoreTransitions.h"
 
-// All shader stages supported by VK device - VK_PIPELINE_STAGE_VERTEX_SHADER_BIT, FRAGMENT etc
-extern uint32 GVulkanDevicePipelineStageBits;
+// All the graphic shader stages supported by the VkDevice (VK_PIPELINE_STAGE_VERTEX_SHADER_BIT, VK_PIPELINE_STAGE_FRAGMENT_SHADER_BIT, etc)
+extern VkPipelineStageFlags GVulkanDeviceGraphicsPipelineStageBits;
 
 extern int32 GVulkanAllowConcurrentImage;
 
@@ -351,7 +351,7 @@ static void GetVkStageAndAccessFlags(ERHIAccess RHIAccess, FRHITransitionInfo::E
 
 	if (EnumHasAnyFlags(RHIAccess, ERHIAccess::SRVGraphics))
 	{
-		StageFlags |= (GVulkanDevicePipelineStageBits & ~VK_PIPELINE_STAGE_COMPUTE_SHADER_BIT);
+		StageFlags |= GVulkanDeviceGraphicsPipelineStageBits;
 		AccessFlags |= VK_ACCESS_SHADER_READ_BIT;
 		if ((UsageFlags & VK_BUFFER_USAGE_UNIFORM_BUFFER_BIT) != 0)
 		{
@@ -376,7 +376,7 @@ static void GetVkStageAndAccessFlags(ERHIAccess RHIAccess, FRHITransitionInfo::E
 
 	if (EnumHasAnyFlags(RHIAccess, ERHIAccess::UAVGraphics))
 	{
-		StageFlags |= (GVulkanDevicePipelineStageBits & ~VK_PIPELINE_STAGE_COMPUTE_SHADER_BIT);
+		StageFlags |= GVulkanDeviceGraphicsPipelineStageBits;
 		AccessFlags |= VK_ACCESS_SHADER_READ_BIT | VK_ACCESS_SHADER_WRITE_BIT;
 		Layout = VK_IMAGE_LAYOUT_GENERAL;
 
diff --git a/Engine/Source/Runtime/VulkanRHI/Private/VulkanRHI.cpp b/Engine/Source/Runtime/VulkanRHI/Private/VulkanRHI.cpp
--- a/Engine/Source/Runtime/VulkanRHI/Private/VulkanRHI.cpp
+++ b/Engine/Source/Runtime/VulkanRHI/Private/VulkanRHI.cpp
@@ -119,8 +119,8 @@ static FAutoConsoleVariableRef GCVarVulkanAMDCompatibilityMode(
 
 extern TAutoConsoleVariable<int32> GVulkanRayTracingCVar;
 
-// All shader stages supported by VK device - VK_PIPELINE_STAGE_VERTEX_SHADER_BIT, FRAGMENT etc
-uint32 GVulkanDevicePipelineStageBits = 0;
+// All the graphic shader stages supported by the VkDevice (VK_PIPELINE_STAGE_VERTEX_SHADER_BIT, VK_PIPELINE_STAGE_FRAGMENT_SHADER_BIT, etc)
+VkPipelineStageFlags GVulkanDeviceGraphicsPipelineStageBits = 0;
 
 DEFINE_LOG_CATEGORY(LogVulkan)
 
@@ -1020,13 +1020,13 @@ void FVulkanDynamicRHI::InitInstance()
 		GUseTexture3DBulkDataRHI = false;
 
 		// these are supported by all devices
-		GVulkanDevicePipelineStageBits = VK_PIPELINE_STAGE_VERTEX_SHADER_BIT | VK_PIPELINE_STAGE_FRAGMENT_SHADER_BIT | VK_PIPELINE_STAGE_COMPUTE_SHADER_BIT;
+		GVulkanDeviceGraphicsPipelineStageBits = VK_PIPELINE_STAGE_VERTEX_SHADER_BIT | VK_PIPELINE_STAGE_FRAGMENT_SHADER_BIT;
 		VkShaderStageFlags VulkanDeviceShaderStageBits = VK_SHADER_STAGE_VERTEX_BIT | VK_SHADER_STAGE_FRAGMENT_BIT | VK_SHADER_STAGE_COMPUTE_BIT;
 
 		// optional shader stages
 		if (Features.Core_1_0.geometryShader)
 		{
-			GVulkanDevicePipelineStageBits |= VK_PIPELINE_STAGE_GEOMETRY_SHADER_BIT;
+			GVulkanDeviceGraphicsPipelineStageBits |= VK_PIPELINE_STAGE_GEOMETRY_SHADER_BIT;
 			VulkanDeviceShaderStageBits |= VK_SHADER_STAGE_GEOMETRY_BIT;
 		}
 
@@ -1037,7 +1037,7 @@ void FVulkanDynamicRHI::InitInstance()
 			GRHIGlobals.SupportsMeshShadersTier0 = RHISupportsMeshShadersTier0(GMaxRHIShaderPlatform);
 			GRHIGlobals.SupportsMeshShadersTier1 = RHISupportsMeshShadersTier1(GMaxRHIShaderPlatform);
 
-			GVulkanDevicePipelineStageBits |= VK_PIPELINE_STAGE_TASK_SHADER_BIT_EXT | VK_PIPELINE_STAGE_MESH_SHADER_BIT_EXT;
+			GVulkanDeviceGraphicsPipelineStageBits |= VK_PIPELINE_STAGE_TASK_SHADER_BIT_EXT | VK_PIPELINE_STAGE_MESH_SHADER_BIT_EXT;
 			VulkanDeviceShaderStageBits |= VK_SHADER_STAGE_TASK_BIT_EXT | VK_SHADER_STAGE_MESH_BIT_EXT;
 		}
 #endif
