[ASIS] Skip frame if window was not created

#rnx

#Synced-CL 0

diff --git a/Engine/Source/Runtime/VulkanRHI/Private/Android/VulkanAndroidPlatform.cpp b/Engine/Source/Runtime/VulkanRHI/Private/Android/VulkanAndroidPlatform.cpp
--- a/Engine/Source/Runtime/VulkanRHI/Private/Android/VulkanAndroidPlatform.cpp
+++ b/Engine/Source/Runtime/VulkanRHI/Private/Android/VulkanAndroidPlatform.cpp
@@ -656,6 +656,10 @@ extern bool IsInAndroidEventThread();
 
 namespace AndroidVulkan
 {
+#if USE_ANDROID_STANDALONE
+	static std::atomic<bool> bWaitingForWindowCreated{false};
+#endif
+
 	static void OnLostWindow(TOptional<FAndroidWindow::FNativeAccessor> WindowContainer)
 	{
 		check(IsInAndroidEventThread());
@@ -714,6 +718,10 @@ namespace AndroidVulkan
 		FVulkanPlatformWindowContext WindowContext(MoveTemp(WindowContainer));
 
 		UE_LOG(LogVulkanRHI, Log, TEXT("OnFoundWindow:"));
+
+#if USE_ANDROID_STANDALONE
+		bWaitingForWindowCreated.store(true);
+#endif
 		FGraphEventRef OnComplete = FGraphEvent::CreateGraphEvent();
 		ENQUEUE_RENDER_COMMAND(OnAndroidFoundWindow)([OnComplete, &WindowContext](FRHICommandListImmediate& RHICmdList) mutable
 			{
@@ -757,6 +765,12 @@ namespace AndroidVulkan
 	}
 }
 
+#if USE_ANDROID_STANDALONE
+bool FVulkanAndroidPlatform::IsWaitingForWindowCreated()
+{
+	return AndroidVulkan::bWaitingForWindowCreated.load();
+}
+#endif
 
 // Old window behavior window recreate function.
 // Moved to platform specific code from VulkanRHI.cpp.
@@ -837,6 +851,24 @@ void FVulkanAndroidPlatform::OverridePlatformHandlers(bool bInit)
 		{
 			UE::FAndroidPlatformDynamicRHI::SetRHIOnReleaseWindowCallback(AndroidVulkan::OnLostWindow);
 			UE::FAndroidPlatformDynamicRHI::SetRHIOnReInitWindowCallback(AndroidVulkan::OnFoundWindow);
+#if USE_ANDROID_STANDALONE
+			UE::FAndroidPlatformDynamicRHI::SetRHIOnWindowCreatedCallback(
+				[]()
+				{
+					// Enqueue on the render thread so this executes AFTER the viewport resize
+					// render command (triggered by OnWindowSizeChanged) has completed.
+					// Unblocking immediately here would race with the resize, allowing a frame
+					// to present with the old swapchain size.
+					ENQUEUE_RENDER_COMMAND(OnVulkanWindowCreated)([](FRHICommandListImmediate& RHICmdList)
+					{
+						RHICmdList.EnqueueLambda([](FRHICommandListImmediate&)
+						{
+							AndroidVulkan::bWaitingForWindowCreated.store(false);
+							UE_LOG(LogVulkanRHI, Display, TEXT("OnWindowCreatedProcessed: rendering unblocked"));
+						});
+					});
+				});
+#endif
 		}
 		else
 		{
@@ -851,6 +883,7 @@ void FVulkanAndroidPlatform::OverridePlatformHandlers(bool bInit)
 		FPlatformMisc::SetOnReleaseWindowCallback(nullptr);
 		UE::FAndroidPlatformDynamicRHI::SetRHIOnReleaseWindowCallback(nullptr);
 		UE::FAndroidPlatformDynamicRHI::SetRHIOnReInitWindowCallback(nullptr);
+		UE::FAndroidPlatformDynamicRHI::SetRHIOnWindowCreatedCallback(nullptr);
 
 		FPlatformMisc::SetOnPauseCallback(nullptr);
 	}
diff --git a/Engine/Source/Runtime/VulkanRHI/Private/Android/VulkanAndroidPlatform.h b/Engine/Source/Runtime/VulkanRHI/Private/Android/VulkanAndroidPlatform.h
--- a/Engine/Source/Runtime/VulkanRHI/Private/Android/VulkanAndroidPlatform.h
+++ b/Engine/Source/Runtime/VulkanRHI/Private/Android/VulkanAndroidPlatform.h
@@ -184,6 +184,9 @@ public:
 	static bool SupportsVolumeTextureRendering() { return false; }
 
 	static void OverridePlatformHandlers(bool bInit);
+#if USE_ANDROID_STANDALONE
+	static bool IsWaitingForWindowCreated();
+#endif
 
 	//#todo-rco: Detect Mali?
 	static bool RequiresPresentLayoutFix() { return true; }
diff --git a/Engine/Source/Runtime/VulkanRHI/Private/VulkanViewport.cpp b/Engine/Source/Runtime/VulkanRHI/Private/VulkanViewport.cpp
--- a/Engine/Source/Runtime/VulkanRHI/Private/VulkanViewport.cpp
+++ b/Engine/Source/Runtime/VulkanRHI/Private/VulkanViewport.cpp
@@ -684,6 +684,14 @@ bool FVulkanViewport::Present(FVulkanCommandListContext& Context, const FRHIPres
 {
 	check(Context.IsImmediate());
 
+#if PLATFORM_ANDROID && USE_ANDROID_STANDALONE
+	if (FVulkanPlatform::IsWaitingForWindowCreated())
+	{
+		UE_LOG(LogVulkanRHI, Log, TEXT("Present: skipping frame, waiting for WINDOW_CREATED"));
+		return true;
+	}
+#endif
+
 	if (!InPresentArgs.bPresent)
 	{
 		return false;
