VulkanRHI: Fix 3D texture layout validation errors with maintenance9 enabled.

#rnx

#Synced-CL 52444545

diff --git a/Engine/Source/Runtime/VulkanRHI/Private/VulkanBarriers.cpp b/Engine/Source/Runtime/VulkanRHI/Private/VulkanBarriers.cpp
--- a/Engine/Source/Runtime/VulkanRHI/Private/VulkanBarriers.cpp
+++ b/Engine/Source/Runtime/VulkanRHI/Private/VulkanBarriers.cpp
@@ -617,32 +617,6 @@ static void SetupImageBarrierEntireRes(VkImageMemoryBarrier2& ImgBarrier, VkImag
 	SetupImageBarrier(ImgBarrier, Image, SrcStageFlags, DstStageFlags, SrcAccessFlags, DstAccessFlags, SrcLayout, DstLayout, SubresRange);
 }
 
-// Fill in a VkImageSubresourceRange struct from the data contained inside a transition info struct coming from the public API.
-static void SetupSubresourceRange(VkImageSubresourceRange& SubresRange, const FRHITransitionInfo& TransitionInfo, VkImageAspectFlags AspectMask)
-{
-	SubresRange.aspectMask = AspectMask;
-	if (TransitionInfo.IsAllMips())
-	{
-		SubresRange.baseMipLevel = 0;
-		SubresRange.levelCount = VK_REMAINING_MIP_LEVELS;
-	}
-	else
-	{
-		SubresRange.baseMipLevel = TransitionInfo.MipIndex;
-		SubresRange.levelCount = 1;
-	}
-
-	if (TransitionInfo.IsAllArraySlices())
-	{
-		SubresRange.baseArrayLayer = 0;
-		SubresRange.layerCount = VK_REMAINING_ARRAY_LAYERS;
-	}
-	else
-	{
-		SubresRange.baseArrayLayer = TransitionInfo.ArraySlice;
-		SubresRange.layerCount = 1;
-	}
-}
 
 void FVulkanDynamicRHI::RHICreateTransition(FRHITransition* Transition, const FRHITransitionCreateInfo& CreateInfo)
 {
@@ -947,7 +921,23 @@ void ConvertTransitionToBarriers(FVulkanCommandListContext& Context, const FVulk
 			ImageBarrier.dstQueueFamilyIndex = VK_QUEUE_FAMILY_IGNORED;
 			ImageBarrier.image = Texture->Image;
 
-			SetupSubresourceRange(ImageBarrier.subresourceRange, Info, AspectFlags);
+			// For 3D textures created with VK_IMAGE_CREATE_2D_ARRAY_COMPATIBLE_BIT when Maintenance9 is enabled,
+			// a layerCount==1 refers to a single slice (used to refer to the rest of the texture).  Maintain the same behavior
+			// with/without Maintenance9 (without altering values for drivers without).
+			const FRHITextureDesc& Desc = Texture->GetDesc();
+			const bool bUseAllSlices =
+				(Desc.Dimension == ETextureDimension::Texture3D) &&
+				Context.Device.GetOptionalExtensions().HasKHRMaintenance9 &&
+				!EnumHasAnyFlags(Desc.Flags, TexCreate_ReservedResource);
+
+			// Fill the ImageBarrier.subresourceRange from the FRHITransitionInfo
+			ImageBarrier.subresourceRange = {
+				.aspectMask = AspectFlags,
+				.baseMipLevel = Info.IsAllMips() ? 0u : Info.MipIndex,
+				.levelCount = Info.IsAllMips() ? VK_REMAINING_MIP_LEVELS : 1u,
+				.baseArrayLayer = Info.IsAllArraySlices() ? 0u : Info.ArraySlice,
+				.layerCount = (Info.IsAllArraySlices() || bUseAllSlices) ? VK_REMAINING_ARRAY_LAYERS : 1u
+			};
 
 			if constexpr (bIsLegacyBarrier)
 			{
