VulkanRHI: Descriptor Heap support for ShaderObjects.  SubGroupSize support for Shader Objects.  Fix validation issues about unset shader object handles on new command buffers.

#rnx

#Synced-CL 52515375

diff --git a/Engine/Source/Runtime/VulkanRHI/Private/VulkanCommandBuffer.cpp b/Engine/Source/Runtime/VulkanRHI/Private/VulkanCommandBuffer.cpp
--- a/Engine/Source/Runtime/VulkanRHI/Private/VulkanCommandBuffer.cpp
+++ b/Engine/Source/Runtime/VulkanRHI/Private/VulkanCommandBuffer.cpp
@@ -299,6 +299,8 @@ void FVulkanCommandBuffer::Begin(FVulkanQueryPool* OptionalQueryPool, VkRenderPa
 		VulkanRHI::vkCmdWriteTimestamp(GetHandle(), VK_PIPELINE_STAGE_TOP_OF_PIPE_BIT, OptionalQueryPool->GetHandle(), IndexInPool);
 	}
 
+	ResetBoundShaderObject();
+
 	bNeedsFullDynamicStateUpdate = true;
 }
 
@@ -322,10 +324,45 @@ void FVulkanCommandBuffer::SetUnusedGraphicsDynamicStates()
 
 		const VkSampleMask SampleMask = MAX_uint32;
 		VulkanRHI::vkCmdSetSampleMaskEXT(CommandBufferHandle, VK_SAMPLE_COUNT_1_BIT, &SampleMask);
+	}
+}
+
+void FVulkanCommandBuffer::ResetBoundShaderObject()
+{
+	if (Device.SupportsShaderObjects())
+	{
+		CurrentComputeShaderObject = VK_NULL_HANDLE;
+
+		// Clear all graphics stages, including the unused tess stages
+		if (CommandBufferPool.GetQueue().GetQueueType() == EVulkanQueueType::Graphics)
+		{
+			FMemory::Memzero(CurrentGraphicsShaderObjects);
+
+			uint32 StageCount = 0;
+			TStaticArray<VkShaderStageFlagBits, ShaderStage::NumGraphicsStages + 2> ShaderStages;
+
+			ShaderStages[StageCount++] = VK_SHADER_STAGE_TESSELLATION_CONTROL_BIT; // unused
+			ShaderStages[StageCount++] = VK_SHADER_STAGE_TESSELLATION_EVALUATION_BIT; // unused
 
-		// Clear unused stages
-		const VkShaderStageFlagBits ShaderStages[] = { VK_SHADER_STAGE_TESSELLATION_CONTROL_BIT, VK_SHADER_STAGE_TESSELLATION_EVALUATION_BIT };
-		VulkanRHI::vkCmdBindShadersEXT(CommandBufferHandle, UE_ARRAY_COUNT(ShaderStages), ShaderStages, nullptr);
+			ShaderStages[StageCount++] = VK_SHADER_STAGE_VERTEX_BIT;
+			ShaderStages[StageCount++] = VK_SHADER_STAGE_FRAGMENT_BIT;
+
+#if PLATFORM_SUPPORTS_GEOMETRY_SHADERS
+			if (Device.GetPhysicalDeviceFeatures().Core_1_0.geometryShader)
+			{
+				ShaderStages[StageCount++] = VK_SHADER_STAGE_GEOMETRY_BIT;
+			}
+#endif
+
+#if PLATFORM_SUPPORTS_MESH_SHADERS
+			if (Device.GetOptionalExtensions().HasEXTMeshShader)
+			{
+				ShaderStages[StageCount++] = VK_SHADER_STAGE_MESH_BIT_EXT;
+				ShaderStages[StageCount++] = VK_SHADER_STAGE_TASK_BIT_EXT;
+			}
+#endif
+			VulkanRHI::vkCmdBindShadersEXT(CommandBufferHandle, StageCount, ShaderStages.GetData(), nullptr);
+		}
 	}
 }
 
@@ -336,8 +373,6 @@ void FVulkanCommandBuffer::ResetCurrentState()
 	CurrentStencilRef = 0;
 
 	CurrentPipelineHandle = VK_NULL_HANDLE;
-	CurrentComputeShaderObject = VK_NULL_HANDLE;
-	FMemory::Memzero(CurrentGraphicsShaderObjects);
 
 	CurrentPushConstants.SetNumUninitialized(FVulkanBindlessDescriptorManager::RequiredPushDataSize / sizeof(uint32));
 	FMemory::Memset(CurrentPushConstants.GetData(), 0xFF, CurrentPushConstants.NumBytes());
diff --git a/Engine/Source/Runtime/VulkanRHI/Private/VulkanCommandBuffer.h b/Engine/Source/Runtime/VulkanRHI/Private/VulkanCommandBuffer.h
--- a/Engine/Source/Runtime/VulkanRHI/Private/VulkanCommandBuffer.h
+++ b/Engine/Source/Runtime/VulkanRHI/Private/VulkanCommandBuffer.h
@@ -188,6 +188,7 @@ private:
 	TArray<uint32> CurrentPushConstants;
 
 	void ResetCurrentState();
+	void ResetBoundShaderObject();
 
 	TArray<VkEvent> EndedBarrierEvents;
 
diff --git a/Engine/Source/Runtime/VulkanRHI/Private/VulkanShaderObjectManager.cpp b/Engine/Source/Runtime/VulkanRHI/Private/VulkanShaderObjectManager.cpp
--- a/Engine/Source/Runtime/VulkanRHI/Private/VulkanShaderObjectManager.cpp
+++ b/Engine/Source/Runtime/VulkanRHI/Private/VulkanShaderObjectManager.cpp
@@ -91,6 +91,40 @@ void FVulkanShaderObjectManager::InternalCreateShaderObject(FVulkanShader& Vulka
 		.pSpecializationInfo = nullptr
 	};
 
+	VkShaderDescriptorSetAndBindingMappingInfoEXT BindingMappingInfo;
+	if (Device.GetBindlessDescriptorManager()->UseDescriptorHeaps())
+	{
+		ShaderCreateInfo.flags |= VK_SHADER_CREATE_DESCRIPTOR_HEAP_BIT_EXT;
+
+		TConstArrayView<VkDescriptorSetAndBindingMappingEXT> BindingMappings = Device.GetBindlessDescriptorManager()->GetBindingMappings(VulkanShader.Frequency, (VulkanShader.GetCodeHeader().PackedGlobalsSize > 0));
+		BindingMappingInfo = {
+			.sType = VK_STRUCTURE_TYPE_SHADER_DESCRIPTOR_SET_AND_BINDING_MAPPING_INFO_EXT,
+			.pNext = nullptr,
+			.mappingCount = (uint32)BindingMappings.Num(),
+			.pMappings = BindingMappings.GetData()
+		};
+		AddToPNext(ShaderCreateInfo, BindingMappingInfo);
+	}
+
+	const uint32 WaveSize = VulkanShader.GetCodeHeader().WaveSize;
+	VkPipelineShaderStageRequiredSubgroupSizeCreateInfo RequiredSubgroupSizeCreateInfo;
+	if ((WaveSize > 0) && Device.GetOptionalExtensions().HasEXTSubgroupSizeControl)
+	{
+		const VkPhysicalDeviceSubgroupSizeControlPropertiesEXT& SubgroupSizeControlProperties = Device.GetOptionalExtensionProperties().SubgroupSizeControlProperties;
+		const bool bSupportedStage = (VKHasAllFlags(SubgroupSizeControlProperties.requiredSubgroupSizeStages, Stage));
+		const bool bSupportedSize = ((WaveSize >= SubgroupSizeControlProperties.minSubgroupSize) && (WaveSize <= SubgroupSizeControlProperties.maxSubgroupSize));
+		if (bSupportedStage && bSupportedSize)
+		{
+			RequiredSubgroupSizeCreateInfo = {
+				.sType = VK_STRUCTURE_TYPE_PIPELINE_SHADER_STAGE_REQUIRED_SUBGROUP_SIZE_CREATE_INFO,
+				.pNext = nullptr,
+				.requiredSubgroupSize = WaveSize
+			};
+
+			AddToPNext(ShaderCreateInfo, RequiredSubgroupSizeCreateInfo);
+		}
+	}
+
 	// :todo-jn: Task shaders never used at the moment
 	if (Stage == VK_SHADER_STAGE_MESH_BIT_EXT)
 	{
