VulkanRHI: Fix default sample mask value for shader objects.

#rnx

#Synced-CL 52539552

diff --git a/Engine/Source/Runtime/VulkanRHI/Private/VulkanCommandBuffer.cpp b/Engine/Source/Runtime/VulkanRHI/Private/VulkanCommandBuffer.cpp
--- a/Engine/Source/Runtime/VulkanRHI/Private/VulkanCommandBuffer.cpp
+++ b/Engine/Source/Runtime/VulkanRHI/Private/VulkanCommandBuffer.cpp
@@ -322,8 +322,8 @@ void FVulkanCommandBuffer::SetUnusedGraphicsDynamicStates()
 		const float blendConstants[4] = {1.0f, 1.0f, 1.0f, 1.0f};
 		VulkanRHI::vkCmdSetBlendConstants(CommandBufferHandle, blendConstants);
 
-		const VkSampleMask SampleMask = MAX_uint32;
-		VulkanRHI::vkCmdSetSampleMaskEXT(CommandBufferHandle, VK_SAMPLE_COUNT_1_BIT, &SampleMask);
+		const VkSampleMask SampleMask[] = { MAX_uint32, MAX_uint32 };
+		VulkanRHI::vkCmdSetSampleMaskEXT(CommandBufferHandle, VK_SAMPLE_COUNT_64_BIT, SampleMask);
 	}
 }
 
