Vulkan Sampler States
- SamplerStates bindless handles weren't being freed on shutdown.
- Changing the SamplerMap from TMap<uint32, FSamplerStateRHIRef> to TMap<uint32, TRefCountPtr<FVulkanSamplerState>>

#Synced-CL 52522322

diff --git a/Engine/Source/Runtime/VulkanRHI/Private/VulkanDevice.h b/Engine/Source/Runtime/VulkanRHI/Private/VulkanDevice.h
--- a/Engine/Source/Runtime/VulkanRHI/Private/VulkanDevice.h
+++ b/Engine/Source/Runtime/VulkanRHI/Private/VulkanDevice.h
@@ -503,7 +503,7 @@ public:
 		return BindlessDescriptorManager;
 	}
 
-	TMap<uint32, FSamplerStateRHIRef>& GetSamplerMap()
+	TMap<uint32, TRefCountPtr<FVulkanSamplerState>>& GetSamplerMap()
 	{
 		return SamplerMap;
 	}
@@ -716,7 +716,7 @@ private:
 
 	VkComponentMapping PixelFormatComponentMapping[PF_MAX];
 
-	TMap<uint32, FSamplerStateRHIRef> SamplerMap;
+	TMap<uint32, TRefCountPtr<FVulkanSamplerState>> SamplerMap;
 
 	FVulkanCommandListContext* ImmediateContext;
 
diff --git a/Engine/Source/Runtime/VulkanRHI/Private/VulkanRHI.cpp b/Engine/Source/Runtime/VulkanRHI/Private/VulkanRHI.cpp
--- a/Engine/Source/Runtime/VulkanRHI/Private/VulkanRHI.cpp
+++ b/Engine/Source/Runtime/VulkanRHI/Private/VulkanRHI.cpp
@@ -534,14 +534,17 @@ void FVulkanDynamicRHI::Shutdown()
 		// Ask all initialized FRenderResources to release their RHI resources.
 		FRenderResource::ReleaseRHIForAllResources();
 
+		for (const TPair<uint32, TRefCountPtr<FVulkanSamplerState>>& Pair : Device->SamplerMap)
 		{
-			for (auto& Pair : Device->SamplerMap)
+			FVulkanSamplerState* SamplerState = Pair.Value.GetReference();
+			VulkanRHI::vkDestroySampler(Device->GetHandle(), SamplerState->Sampler, VULKAN_CPU_ALLOCATOR);
+
+			if (SamplerState->BindlessHandle.IsValid())
 			{
-				FVulkanSamplerState* SamplerState = (FVulkanSamplerState*)Pair.Value.GetReference();
-				VulkanRHI::vkDestroySampler(Device->GetHandle(), SamplerState->Sampler, VULKAN_CPU_ALLOCATOR);
+				Device->GetBindlessDescriptorManager()->FreeDescriptor(SamplerState->BindlessHandle);
 			}
-			Device->SamplerMap.Empty();
 		}
+		Device->SamplerMap.Empty();
 
 		Device->CleanUpRayTracing();
 
diff --git a/Engine/Source/Runtime/VulkanRHI/Private/VulkanState.cpp b/Engine/Source/Runtime/VulkanRHI/Private/VulkanState.cpp
--- a/Engine/Source/Runtime/VulkanRHI/Private/VulkanState.cpp
+++ b/Engine/Source/Runtime/VulkanRHI/Private/VulkanState.cpp
@@ -355,16 +355,16 @@ FSamplerStateRHIRef FVulkanDynamicRHI::RHICreateSamplerState(const FSamplerState
 
 	{
 		FScopeLock ScopeLock(&GSamplerHashLock);
-		TMap<uint32, FSamplerStateRHIRef>& SamplerMap = Device->GetSamplerMap();
-		FSamplerStateRHIRef* Found = SamplerMap.Find(CRC);
-		if (Found)
+		TMap<uint32, TRefCountPtr<FVulkanSamplerState>>& SamplerMap = Device->GetSamplerMap();
+
+		if (TRefCountPtr<FVulkanSamplerState>* Found = SamplerMap.Find(CRC))
 		{
-			return *Found;
+			return FSamplerStateRHIRef(*Found);
 		}
 
-		FSamplerStateRHIRef New = new FVulkanSamplerState(SamplerInfo, *Device);
+		TRefCountPtr<FVulkanSamplerState> New = new FVulkanSamplerState(SamplerInfo, *Device);
 		SamplerMap.Add(CRC, New);
-		return New;
+		return FSamplerStateRHIRef(MoveTemp(New));
 	}
 }
 
