VulkanRHI: wait for async shader object tasks on shutdown.

#rnx

#Synced-CL 52548173

diff --git a/Engine/Source/Runtime/VulkanRHI/Private/VulkanShaderObjectManager.cpp b/Engine/Source/Runtime/VulkanRHI/Private/VulkanShaderObjectManager.cpp
--- a/Engine/Source/Runtime/VulkanRHI/Private/VulkanShaderObjectManager.cpp
+++ b/Engine/Source/Runtime/VulkanRHI/Private/VulkanShaderObjectManager.cpp
@@ -22,6 +22,24 @@ FVulkanShaderObjectManager::FVulkanShaderObjectManager(FVulkanDevice& InDevice)
 {
 }
 
+FVulkanShaderObjectManager::~FVulkanShaderObjectManager()
+{
+	// Move pending async jobs into an array we won't need to lock, and wait for them to complete
+	TArray<TRefCountPtr<FVulkanShader>> ShaderArray;
+	{
+		FScopeLock ScopeLock(&AsyncTaskLock);
+		ShaderArray = AsyncTasks.Array();
+		AsyncTasks.Empty();
+	}
+	for (TRefCountPtr<FVulkanShader>& ShaderRef : ShaderArray)
+	{
+		if (ShaderRef->AsyncCompileTask.IsValid() && !ShaderRef->AsyncCompileTask->IsComplete())
+		{
+			ShaderRef->AsyncCompileTask->Wait();
+		}
+	}
+}
+
 VkShaderStageFlags FVulkanShaderObjectManager::GetPossibleNextStages(VkShaderStageFlagBits Stage)
 {
 	switch (Stage)
@@ -50,10 +68,12 @@ void FVulkanShaderObjectManager::CreateShaderObject(FVulkanShader& VulkanShader)
 {
 	if (GAllowAsyncShaderObjectCompilation)
 	{
+		FScopeLock ScopeLock(&AsyncTaskLock);
 		VulkanShader.AsyncCompileTask = FFunctionGraphTask::CreateAndDispatchWhenReady([ShaderObjectManager = this, &VulkanShader]()
 		{
 			ShaderObjectManager->InternalCreateShaderObject(VulkanShader);
 		});
+		AsyncTasks.Emplace(&VulkanShader);
 	}
 	else
 	{
@@ -154,6 +174,11 @@ void FVulkanShaderObjectManager::InternalCreateShaderObject(FVulkanShader& Vulka
 	{
 		VULKAN_SET_DEBUG_NAME(Device, VK_OBJECT_TYPE_SHADER_EXT, VulkanShader.ShaderObject, TEXT("%s"), *VulkanShader.GetCodeHeader().DebugName);
 	}
+
+	{
+		FScopeLock ScopeLock(&AsyncTaskLock);
+		AsyncTasks.Remove(&VulkanShader);
+	}
 }
 
 bool FVulkanShaderObjectManager::CreateLinkedShaderObjects(FVulkanGraphicsPipelineState* PSO, FVulkanShader* Shaders[ShaderStage::NumGraphicsStages], FGraphicsPipelineStateInitializer::EPSOPrecacheCompileType PSOCompileType)
diff --git a/Engine/Source/Runtime/VulkanRHI/Private/VulkanShaderObjectManager.h b/Engine/Source/Runtime/VulkanRHI/Private/VulkanShaderObjectManager.h
--- a/Engine/Source/Runtime/VulkanRHI/Private/VulkanShaderObjectManager.h
+++ b/Engine/Source/Runtime/VulkanRHI/Private/VulkanShaderObjectManager.h
@@ -15,6 +15,7 @@ class FVulkanShaderObjectManager
 {
 public:
 	FVulkanShaderObjectManager(FVulkanDevice& InDevice);
+	~FVulkanShaderObjectManager();
 
 	void CreateShaderObject(FVulkanShader& VulkanShader);
 	bool CreateLinkedShaderObjects(FVulkanGraphicsPipelineState* PSO, FVulkanShader* Shaders[ShaderStage::NumGraphicsStages], FGraphicsPipelineStateInitializer::EPSOPrecacheCompileType PSOCompileType);
@@ -24,4 +25,7 @@ private:
 	void InternalCreateShaderObject(FVulkanShader& VulkanShader);
 
 	FVulkanDevice& Device;
+
+	FCriticalSection AsyncTaskLock;
+	TSet<TRefCountPtr<FVulkanShader>> AsyncTasks;
 };
\ No newline at end of file
