VulkanRHI: Properly initialize dynamic state when FragmentShadingRate is enabled but unused with shader objects.

#rnx

#Synced-CL 52581599

diff --git a/Engine/Source/Runtime/VulkanRHI/Private/VulkanPendingState.cpp b/Engine/Source/Runtime/VulkanRHI/Private/VulkanPendingState.cpp
--- a/Engine/Source/Runtime/VulkanRHI/Private/VulkanPendingState.cpp
+++ b/Engine/Source/Runtime/VulkanRHI/Private/VulkanPendingState.cpp
@@ -795,9 +795,9 @@ void FVulkanPendingGfxState::UpdateDynamicStates(FVulkanCommandBuffer& CommandBu
 		VulkanRHI::vkCmdSetStencilWriteMask(CommandBuffer.GetHandle(), VK_STENCIL_FACE_BACK_BIT, NewDynamicState.DepthStencil.BackWriteMask);
 		VULKAN_DYNAMIC_STATE_BLOCK_END();
 
-		if (GRHISupportsPipelineVariableRateShading && (GRHIVariableRateShadingImageDataType == VRSImage_Palette))
+		VULKAN_DYNAMIC_STATE_BLOCK_START(ShadingRate, RenderTargets.bHasFragmentDensityAttachment, Combiner);
 		{
-			VULKAN_DYNAMIC_STATE_BLOCK_START(ShadingRate, RenderTargets.bHasFragmentDensityAttachment, Combiner);
+			if (GRHISupportsPipelineVariableRateShading && (GRHIVariableRateShadingImageDataType == VRSImage_Palette))
 			{
 				const bool bNewUsingVariableRateShading = (NewDynamicState.ShadingRate != EVRSShadingRate::VRSSR_1x1) ||
 					(NewDynamicState.RenderTargets.bHasFragmentDensityAttachment && NewDynamicState.Combiner != EVRSRateCombiner::VRSRB_Passthrough);
@@ -821,8 +821,14 @@ void FVulkanPendingGfxState::UpdateDynamicStates(FVulkanCommandBuffer& CommandBu
 					VulkanRHI::vkCmdSetFragmentShadingRateKHR(CommandBuffer.GetHandle(), &FragmentSize, CombinerOps);
 				}
 			}
-			VULKAN_DYNAMIC_STATE_BLOCK_END();
+			else if (bForceUpdate && Device.GetOptionalExtensionProperties().FragmentShadingRateFeatures.pipelineFragmentShadingRate)
+			{
+				const VkExtent2D FragmentSize = { 1, 1 };
+				const VkFragmentShadingRateCombinerOpKHR CombinerOps[] = { VK_FRAGMENT_SHADING_RATE_COMBINER_OP_KEEP_KHR , VK_FRAGMENT_SHADING_RATE_COMBINER_OP_KEEP_KHR };
+				VulkanRHI::vkCmdSetFragmentShadingRateKHR(CommandBuffer.GetHandle(), &FragmentSize, CombinerOps);
+			}
 		}
+		VULKAN_DYNAMIC_STATE_BLOCK_END();
 	}
 
 #undef VULKAN_DYNAMIC_STATE_COMPARE
