VulkanRHI: print value returned by vkGetDescriptorSetLayoutSizeEXT when it doesn't match and convert to warning.

#rnx

#Synced-CL 52581599

diff --git a/Engine/Source/Runtime/VulkanRHI/Private/VulkanBindlessDescriptorManager.cpp b/Engine/Source/Runtime/VulkanRHI/Private/VulkanBindlessDescriptorManager.cpp
--- a/Engine/Source/Runtime/VulkanRHI/Private/VulkanBindlessDescriptorManager.cpp
+++ b/Engine/Source/Runtime/VulkanRHI/Private/VulkanBindlessDescriptorManager.cpp
@@ -764,7 +764,17 @@ void FVulkanBindlessDescriptorManager::InitDescriptorBuffers()
 			VkDeviceSize LayoutSizeInBytes = 0;
 			VulkanRHI::vkGetDescriptorSetLayoutSizeEXT(DeviceHandle, InOutState.DescriptorSetLayout, &LayoutSizeInBytes);
 			// Double check that the layout follows the rules for a single binding with an array of descriptors that are tightly packed
-			check(LayoutSizeInBytes == DescriptorBufferSize);
+			if (LayoutSizeInBytes == 0)
+			{
+				// Reported this way by older drivers, log it as a warning.
+				UE_LOGF(LogRHI, Warning, "vkGetDescriptorSetLayoutSizeEXT returned 0! Expected size: %u.", DescriptorBufferSize);
+			}
+			else
+			{
+				checkf((LayoutSizeInBytes == DescriptorBufferSize),
+					TEXT("vkGetDescriptorSetLayoutSizeEXT returned %u instead of expected size of %u."),
+					(uint32)LayoutSizeInBytes, DescriptorBufferSize);
+			}
 
 			if (IsSamplerSet)
 			{
