VulkanRHI: fix cross queue ownership transfers using maintenance9

#rnx

#Synced-CL 52581599

diff --git a/Engine/Source/Runtime/VulkanRHI/Private/VulkanBarriers.cpp b/Engine/Source/Runtime/VulkanRHI/Private/VulkanBarriers.cpp
--- a/Engine/Source/Runtime/VulkanRHI/Private/VulkanBarriers.cpp
+++ b/Engine/Source/Runtime/VulkanRHI/Private/VulkanBarriers.cpp
@@ -630,7 +630,8 @@ void FVulkanDynamicRHI::RHICreateTransition(FRHITransition* Transition, const FR
 	Data->SrcPipelines = SrcPipelines;
 	Data->DstPipelines = DstPipelines;
 
-	if ((SrcPipelines != DstPipelines) && !EnumHasAnyFlags(CreateInfo.Flags, ERHITransitionCreateFlags::NoFence))
+	// If queue ownership transfers are necessary, always use fences on cross queue to ensure they are submitted in order
+	if ((SrcPipelines != DstPipelines) && (Device->NeedsQueueFamilyOwnershipTransfers() || !EnumHasAnyFlags(CreateInfo.Flags, ERHITransitionCreateFlags::NoFence)))
 	{
 		Data->Semaphore = new FVulkanSemaphore(*Device);
 	}
@@ -1408,7 +1409,7 @@ void ProcessTransitionSync2(FVulkanCommandListContext& Context, TArrayView<const
 
 	TArray<VkImageMemoryBarrier2> TempImageBarriers;
 
-	const bool bUseOwnershipTransfers = Context.Device.HasAsyncComputeQueue() && !GVulkanAllowConcurrentImage;
+	const bool bUseOwnershipTransfers = Context.Device.HasAsyncComputeQueue() && !GVulkanAllowConcurrentImage && Context.Device.NeedsQueueFamilyOwnershipTransfers();
 	for (const FRHITransition* Transition : Transitions)
 	{
 		const FVulkanTransitionData* Data = Transition->GetPrivateData<FVulkanTransitionData>();
@@ -1456,18 +1457,27 @@ void ProcessTransitionSync2(FVulkanCommandListContext& Context, TArrayView<const
 		if (IsSingleRHIPipeline(Data->SrcPipelines) && IsSingleRHIPipeline(Data->DstPipelines))
 		{
 			// For cross-queue 1..1 transitions we will keep the layout change in the same barrier as the queue transfer
-			if (bUseOwnershipTransfers && (Data->SrcPipelines != Data->DstPipelines))
+			if (Data->SrcPipelines != Data->DstPipelines)
 			{
-				// For buffer transitions, just remove unsupported flags for the compute queue
-				if (ExecutingPipeline == ERHIPipeline::AsyncCompute)
+				if (bUseOwnershipTransfers)
 				{
-					MaskSupportedAsyncFlags(Context.Device, BarrierArrays.BufferBarriers, true, true, true);
-				}
+					// For buffer transitions, just remove unsupported flags for the compute queue
+					if (ExecutingPipeline == ERHIPipeline::AsyncCompute)
+					{
+						MaskSupportedAsyncFlags(Context.Device, BarrierArrays.BufferBarriers, true, true, true);
+					}
 
-				if (!GVulkanAllowConcurrentImage)
-				{
 					PatchCrossPipeTransitions(BarrierArrays.ImageBarriers, Context, Data->SrcPipelines, Data->DstPipelines, bIsBeginTransition);
 				}
+				else
+				{
+					// When we don't have ownership transfers, only submit the barrier once on the graphics queue (ensures all stages are supported)
+					check((TargetPipeline == ERHIPipeline::Graphics) || (OtherPipeline == ERHIPipeline::Graphics));
+					if (ExecutingPipeline != ERHIPipeline::Graphics)
+					{
+						continue;
+					}
+				}
 			}
 
 			SubmitBarriers(MakeArrayView(BarrierArrays.MemoryBarriers), MakeArrayView(BarrierArrays.BufferBarriers), MakeArrayView(BarrierArrays.ImageBarriers), Data->EventHandle);
@@ -1587,6 +1597,12 @@ static void ProcessCrossQueueSemaphores(FVulkanCommandListContext& Context, TArr
 }
 
 
+
+
+
+
+
+
 void FVulkanCommandListContext::RHIBeginTransitions(TArrayView<const FRHITransition*> Transitions)
 {
 	if (Device.SupportsParallelRendering())
@@ -1606,7 +1622,6 @@ void FVulkanCommandListContext::RHIBeginTransitions(TArrayView<const FRHITransit
 	{
 		// Nothing to do for legacy barriers on begin (no split support, no async compute support)
 	}
-
 }
 
 void FVulkanCommandListContext::RHIEndTransitions(TArrayView<const FRHITransition*> Transitions)
diff --git a/Engine/Source/Runtime/VulkanRHI/Private/VulkanDevice.cpp b/Engine/Source/Runtime/VulkanRHI/Private/VulkanDevice.cpp
--- a/Engine/Source/Runtime/VulkanRHI/Private/VulkanDevice.cpp
+++ b/Engine/Source/Runtime/VulkanRHI/Private/VulkanDevice.cpp
@@ -469,6 +469,47 @@ void FVulkanDevice::CreateDevice(TArray<const ANSICHAR*>& DeviceLayers, FVulkanD
 		ActiveQueueFamilies.Add(TransferQueueFamilyIndex);
 	}
 
+	if (HasAsyncComputeQueue() && OptionalDeviceExtensions.HasKHRMaintenance9)
+	{
+		TArray<VkQueueFamilyProperties2> QueueFamilyProperties2;
+		QueueFamilyProperties2.SetNumZeroed(QueueFamilyProps.Num());
+		TArray<VkQueueFamilyOwnershipTransferPropertiesKHR> OwnershipTransferProperties;
+		OwnershipTransferProperties.SetNumUninitialized(QueueFamilyProperties2.Num());
+
+		for (int32 QueueIndex = 0; QueueIndex < QueueFamilyProperties2.Num(); ++QueueIndex)
+		{
+			OwnershipTransferProperties[QueueIndex] = {
+				.sType = VK_STRUCTURE_TYPE_QUEUE_FAMILY_OWNERSHIP_TRANSFER_PROPERTIES_KHR,
+				.pNext = nullptr,
+				.optimalImageTransferToQueueFamilies = 0
+			};
+
+			QueueFamilyProperties2[QueueIndex] = {
+				.sType = VK_STRUCTURE_TYPE_QUEUE_FAMILY_PROPERTIES_2,
+				.pNext = &OwnershipTransferProperties[QueueIndex],
+			};
+		}
+
+		uint32 QueueCount = QueueFamilyProperties2.Num();
+		VulkanRHI::vkGetPhysicalDeviceQueueFamilyProperties2(Gpu, &QueueCount, QueueFamilyProperties2.GetData());
+
+		const bool bGraphicsOwnership = OwnershipTransferProperties[GfxQueueFamilyIndex].optimalImageTransferToQueueFamilies & (1u << ComputeQueueFamilyIndex);
+		const bool bAsyncOwnership = OwnershipTransferProperties[ComputeQueueFamilyIndex].optimalImageTransferToQueueFamilies & (1u << GfxQueueFamilyIndex);
+		const bool bTransferOwnership = (TransferQueueFamilyIndex == -1) || (
+			(OwnershipTransferProperties[TransferQueueFamilyIndex].optimalImageTransferToQueueFamilies & (1u << GfxQueueFamilyIndex)) &&
+			(OwnershipTransferProperties[TransferQueueFamilyIndex].optimalImageTransferToQueueFamilies & (1u << ComputeQueueFamilyIndex)));
+		if (bGraphicsOwnership && bAsyncOwnership && bTransferOwnership)
+		{
+			// Queue ownership transfers aren't necessary between the used queue families
+			bNeedsQueueFamilyOwnershipTransfers = false;
+		}
+		else
+		{
+			// Do not merge pipelines together when queue family ownership transfers are necessary
+			GRHIMultiPipelineMergeableAccessMask = ERHIAccess::Unknown;
+		}
+	}
+
 	// Enumerate the available shading rates
 	if (OptionalDeviceExtensions.HasKHRFragmentShadingRate)
 	{
diff --git a/Engine/Source/Runtime/VulkanRHI/Private/VulkanDevice.h b/Engine/Source/Runtime/VulkanRHI/Private/VulkanDevice.h
--- a/Engine/Source/Runtime/VulkanRHI/Private/VulkanDevice.h
+++ b/Engine/Source/Runtime/VulkanRHI/Private/VulkanDevice.h
@@ -614,6 +614,11 @@ public:
 		return QueueFamilyProps;
 	}
 
+	bool NeedsQueueFamilyOwnershipTransfers() const
+	{
+		return bNeedsQueueFamilyOwnershipTransfers;
+	}
+
 	FVulkanTransientHeapCache& GetOrCreateTransientHeapCache();
 
 	const TArray<const ANSICHAR*>& GetDeviceExtensions() { return DeviceExtensions; }
@@ -694,6 +699,8 @@ private:
 	FVulkanPhysicalDeviceFeatures PhysicalDeviceFeatures;
 
 	TArray<VkQueueFamilyProperties> QueueFamilyProps;
+	bool bNeedsQueueFamilyOwnershipTransfers = true;
+
 	VkFormatProperties FormatProperties[VK_FORMAT_RANGE_SIZE];
 	// Info for formats that are not in the core Vulkan spec (i.e. extensions)
 	mutable TMap<VkFormat, VkFormatProperties> ExtensionFormatProperties;
