VulkanRHI: Add check for extension before querying for address.  Use existing UpdateBindlessHandle function on creation.

#rnx

#Synced-CL 51911432

diff --git a/Engine/Source/Runtime/VulkanRHI/Private/VulkanResources.h b/Engine/Source/Runtime/VulkanRHI/Private/VulkanResources.h
--- a/Engine/Source/Runtime/VulkanRHI/Private/VulkanResources.h
+++ b/Engine/Source/Runtime/VulkanRHI/Private/VulkanResources.h
@@ -1000,7 +1000,7 @@ public:
 		return BindlessHandle;
 	}
 	
-	void UpdateBindlessHandle(const FVulkanContextArray& Contexts);
+	void UpdateBindlessHandle(const FVulkanContextArray& Contexts, bool bImmediate);
 	VkDeviceAddress GetDeviceAddress() const;
 
 	void SetupUniformBufferView();
diff --git a/Engine/Source/Runtime/VulkanRHI/Private/VulkanUniformBuffer.cpp b/Engine/Source/Runtime/VulkanRHI/Private/VulkanUniformBuffer.cpp
--- a/Engine/Source/Runtime/VulkanRHI/Private/VulkanUniformBuffer.cpp
+++ b/Engine/Source/Runtime/VulkanRHI/Private/VulkanUniformBuffer.cpp
@@ -115,7 +115,7 @@ FVulkanUniformBuffer::FVulkanUniformBuffer(FVulkanDevice& InDevice, const FRHIUn
 	{
 		AllocateMemory(Contents);
 		BindlessHandle = Device.GetBindlessDescriptorManager()->AllocateDescriptor(ERHIDescriptorType::CBV);
-		Device.GetBindlessDescriptorManager()->UpdateDescriptor({}, BindlessHandle, GetDeviceAddress(), GetSize(), true);
+		UpdateBindlessHandle({}, true);
 	}
 }
 
@@ -149,25 +149,29 @@ void FVulkanUniformBuffer::UpdateResourceTable(TArrayView<FRHIResource*> Resourc
 	}
 }
 
-void FVulkanUniformBuffer::UpdateBindlessHandle(const FVulkanContextArray& Contexts)
+void FVulkanUniformBuffer::UpdateBindlessHandle(const FVulkanContextArray& Contexts, bool bImmediate)
 {
 	if (BindlessHandle.IsValid() && Device.GetBindlessDescriptorManager()->IsSupported())
 	{
 		const VkDeviceAddress BufferAddress = GetDeviceAddress();
-		Device.GetBindlessDescriptorManager()->UpdateDescriptor(Contexts, BindlessHandle, BufferAddress, GetSize(), false);
+		Device.GetBindlessDescriptorManager()->UpdateDescriptor(Contexts, BindlessHandle, BufferAddress, GetSize(), bImmediate);
 	}
 }
 
 VkDeviceAddress FVulkanUniformBuffer::GetDeviceAddress() const
 {
-	const VkBufferDeviceAddressInfo BufferInfo = {
-		.sType = VK_STRUCTURE_TYPE_BUFFER_DEVICE_ADDRESS_INFO,
-		.pNext = nullptr,
-		.buffer = GetBufferHandle()
-	};
-	VkDeviceAddress BufferAddress = VulkanRHI::vkGetBufferDeviceAddressKHR(Device.GetHandle(), &BufferInfo);
-	BufferAddress += GetOffset();
-	return BufferAddress;
+	if (Device.GetOptionalExtensions().HasBufferDeviceAddress)
+	{
+		const VkBufferDeviceAddressInfo BufferInfo = {
+			.sType = VK_STRUCTURE_TYPE_BUFFER_DEVICE_ADDRESS_INFO,
+			.pNext = nullptr,
+			.buffer = GetBufferHandle()
+		};
+		VkDeviceAddress BufferAddress = VulkanRHI::vkGetBufferDeviceAddressKHR(Device.GetHandle(), &BufferInfo);
+		BufferAddress += GetOffset();
+		return BufferAddress;
+	}
+	return 0;
 }
 
 FUniformBufferRHIRef FVulkanDynamicRHI::RHICreateUniformBuffer(const void* Contents, const FRHIUniformBufferLayout* Layout, EUniformBufferUsage Usage, EUniformBufferValidation Validation)
@@ -219,7 +223,7 @@ void FVulkanDynamicRHI::RHIUpdateUniformBuffer(FRHICommandListBase& RHICmdList,
 		if (ConstantBufferSize > 0)
 		{
 			UniformBuffer->AllocateMemory(TempContents);
-			UniformBuffer->UpdateBindlessHandle(Contexts);
+			UniformBuffer->UpdateBindlessHandle(Contexts, false);
 		}
 
 		UniformBuffer->UpdateResourceTable(CmdListResources);
